Couple of improvements for UnrealSync3:
- Moved WinMain to platform specific place.
- Added reload labels button.
- After finishing sync the program gets to the beggining instead of closing.
- Syncing is now copying AutomationToolLauncher.exe to temp file to allow syncing itself.

diff --git a/Engine/Source/Programs/UnrealSync/Private/GUI.cpp b/Engine/Source/Programs/UnrealSync/Private/GUI.cpp
--- a/Engine/Source/Programs/UnrealSync/Private/GUI.cpp
+++ b/Engine/Source/Programs/UnrealSync/Private/GUI.cpp
@@ -688,11 +688,11 @@ public:
 				SNew(STextBlock).Text(FText::FromString("Preview sync?"))
 			];
 
-		CloseButton = SNew(SButton)
+		GoBackButton = SNew(SButton)
 			.IsEnabled(false)
-			.OnClicked(this, &SMainTabWidget::OnCloseButtonClick)
+			.OnClicked(this, &SMainTabWidget::OnGoBackButtonClick)
 			[
-				SNew(STextBlock).Text(FText::FromString("Close"))
+				SNew(STextBlock).Text(FText::FromString("Go back"))
 			];
 
 		LogListView = SNew(SListView<TSharedPtr<FString> >)
@@ -714,6 +714,16 @@ public:
 				+ SVerticalBox::Slot().AutoHeight()
 				[
 					SNew(SHorizontalBox)
+					+ SHorizontalBox::Slot().AutoWidth()
+					[
+						SNew(SButton).HAlign(HAlign_Right)
+						.OnClicked(this, &SMainTabWidget::OnReloadLabels)
+						.IsEnabled(this, &SMainTabWidget::IsReloadLabelsReady)
+						.VAlign(VAlign_Center)
+						[
+							SNew(STextBlock).Text(FText::FromString("Reload Labels"))
+						]
+					]
 					+ SHorizontalBox::Slot().HAlign(HAlign_Fill)
 					+ SHorizontalBox::Slot().AutoWidth()
 					[
@@ -752,7 +762,7 @@ public:
 					+ SHorizontalBox::Slot()
 					+ SHorizontalBox::Slot().AutoWidth()
 					[
-						CloseButton.ToSharedRef()
+						GoBackButton.ToSharedRef()
 					]
 				]
 			];
@@ -762,11 +772,9 @@ public:
 				Switcher.ToSharedRef()
 			];
 
-		DataRefresh();
-
 		FUnrealSync::RegisterOnDataRefresh(FUnrealSync::FOnDataRefresh::CreateRaw(this, &SMainTabWidget::DataRefresh));
 
-		FUnrealSync::StartLoadingData();
+		OnReloadLabels();
 	}
 	END_SLATE_FUNCTION_BUILD_OPTIMIZATION
 
@@ -808,14 +816,17 @@ private:
 	}
 
 	/**
-	 * Function that is called on when Close button is clicked.
-	 * It tells the app to close.
+	 * Function that is called on when go back button is clicked.
+	 * It tells the app to go back to the beginning.
 	 *
 	 * @returns Tells that this event was handled.
 	 */
-	FReply OnCloseButtonClick()
+	FReply OnGoBackButtonClick()
 	{
-		FPlatformMisc::RequestExit(false);
+		Switcher->SetActiveWidgetIndex(0);
+		LogLines.Empty();
+		LogListView->RequestListRefresh();
+		OnReloadLabels();
 
 		return FReply::Handled();
 	}
@@ -831,6 +842,28 @@ private:
 		return GetCurrentSyncCmdLineProvider().IsReadyForSync();
 	}
 
+	/**
+	 * Tells reload labels button if it should be enabled.
+	 *
+	 * @returns True if reload labels button should be enabled. False otherwise.
+	 */
+	bool IsReloadLabelsReady() const
+	{
+		return !FUnrealSync::IsLoadingInProgress();
+	}
+
+	/**
+	 * Launches reload labels background process.
+	 *
+	 * @returns Tells that this event was handled.
+	 */
+	FReply OnReloadLabels()
+	{
+		FUnrealSync::StartLoadingData();
+
+		return FReply::Handled();
+	}
+
 	/**
 	 * Sync button click function. Gets all widget data, constructs command
 	 * line and launches syncing.
@@ -848,7 +881,6 @@ private:
 			FUnrealSync::TerminateLoadingProcess();
 		}
 
-		Switcher->GetActiveWidget()->SetEnabled(false);
 		Switcher->SetActiveWidgetIndex(1);
 		FUnrealSync::LaunchSync(bArtist, bPreview, CommandLine,
 			FUnrealSync::FOnSyncFinished::CreateRaw(this, &SMainTabWidget::SyncingFinished),
@@ -898,7 +930,7 @@ private:
 	 */
 	void SyncingFinished(bool bSuccess)
 	{
-		CloseButton->SetEnabled(true);
+		GoBackButton->SetEnabled(true);
 	}
 
 	/**
@@ -958,8 +990,8 @@ private:
 	TArray<TSharedPtr<FString> > LogLines;
 	/* Log list view. */
 	TSharedPtr<SListView<TSharedPtr<FString> > > LogListView;
-	/* Close button reference. */
-	TSharedPtr<SButton> CloseButton;
+	/* Go back button reference. */
+	TSharedPtr<SButton> GoBackButton;
 
 	/* Array of sync command line providers. */
 	TArray<TSharedRef<ISyncCommandLineProvider> > SyncCommandLineProviders;
diff --git a/Engine/Source/Programs/UnrealSync/Private/P4DataCache.cpp b/Engine/Source/Programs/UnrealSync/Private/P4DataCache.cpp
--- a/Engine/Source/Programs/UnrealSync/Private/P4DataCache.cpp
+++ b/Engine/Source/Programs/UnrealSync/Private/P4DataCache.cpp
@@ -130,3 +130,20 @@ FP4DataLoader::~FP4DataLoader()
 {
 	Terminate();
 }
+
+bool FP4DataLoader::IsInProgress() const
+{
+	return bInProgress;
+}
+
+void FP4DataLoader::Exit()
+{
+	bInProgress = false;
+}
+
+bool FP4DataLoader::Init()
+{
+	bInProgress = true;
+
+	return true;
+}
diff --git a/Engine/Source/Programs/UnrealSync/Private/P4DataCache.h b/Engine/Source/Programs/UnrealSync/Private/P4DataCache.h
--- a/Engine/Source/Programs/UnrealSync/Private/P4DataCache.h
+++ b/Engine/Source/Programs/UnrealSync/Private/P4DataCache.h
@@ -58,6 +58,16 @@ public:
 
 	virtual ~FP4DataLoader();
 
+	/**
+	 * Initializes the runnable object.
+	 *
+	 * This method is called in the context of the thread object that aggregates this, not the
+	 * thread that passes this runnable to a new thread.
+	 *
+	 * @return True if initialization was successful, false otherwise
+	 */
+	virtual bool Init() OVERRIDE;
+
 	/**
 	 * Main function to run for this thread.
 	 *
@@ -65,6 +75,20 @@ public:
 	 */
 	uint32 Run() OVERRIDE;
 
+	/**
+	 * Exits the runnable object.
+	 *
+	 * Called in the context of the aggregating thread to perform any cleanup.
+	 */
+	virtual void Exit();
+
+	/**
+	 * Is this loading thread in progress?
+	 *
+	 * @returns True if this thread is in progress. False otherwise.
+	 */
+	bool IsInProgress() const;
+
 	/**
 	 * Signals that the loading process should be terminated.
 	 */
@@ -77,6 +101,9 @@ private:
 	/* Handle for thread object. */
 	FRunnableThread *Thread;
 
+	/* If this thread is still running. */
+	bool bInProgress;
+
 	/* Signals if the loading process should be terminated. */
 	bool bTerminate;
 };
diff --git a/Engine/Source/Programs/UnrealSync/Private/UnrealSync.cpp b/Engine/Source/Programs/UnrealSync/Private/UnrealSync.cpp
--- a/Engine/Source/Programs/UnrealSync/Private/UnrealSync.cpp
+++ b/Engine/Source/Programs/UnrealSync/Private/UnrealSync.cpp
@@ -10,14 +10,9 @@ DEFINE_LOG_CATEGORY_STATIC(LogUnrealSync, Log, All);
 
 IMPLEMENT_APPLICATION(UnrealSync, "UnrealSync");
 
-/**
- * WinMain, called when the application is started
- */
-int WINAPI WinMain(_In_ HINSTANCE hInInstance, _In_opt_ HINSTANCE hPrevInstance, _In_ LPSTR, _In_ int nCmdShow)
+void RunUnrealSync(const TCHAR* CommandLine)
 {
 	InitGUI();
-
-	return 0;
 }
 
 FString FUnrealSync::GetLatestLabelForGame(const FString& GameName)
@@ -190,6 +185,12 @@ void FUnrealSync::RegisterOnDataRefresh(FOnDataRefresh OnDataRefresh)
 
 void FUnrealSync::StartLoadingData()
 {
+	bLoadingFinished = false;
+	Data.Reset();
+	LoaderThread.Reset();
+
+	OnDataRefresh.ExecuteIfBound();
+
 	LoaderThread = MakeShareable(new FP4DataLoader(
 		FP4DataLoader::FOnLoadingFinished::CreateStatic(&FUnrealSync::OnP4DataLoadingFinished)
 		));
@@ -204,10 +205,7 @@ void FUnrealSync::OnP4DataLoadingFinished(TSharedPtr<FP4DataCache> Data)
 {
 	FUnrealSync::Data = Data;
 
-	if (Data.IsValid())
-	{
-		bLoadingFinished = true;
-	}
+	bLoadingFinished = true;
 
 	OnDataRefresh.ExecuteIfBound();
 }
diff --git a/Engine/Source/Programs/UnrealSync/Private/UnrealSync.h b/Engine/Source/Programs/UnrealSync/Private/UnrealSync.h
--- a/Engine/Source/Programs/UnrealSync/Private/UnrealSync.h
+++ b/Engine/Source/Programs/UnrealSync/Private/UnrealSync.h
@@ -6,6 +6,11 @@
 #include "CoreUObject.h"
 #include "P4DataCache.h"
 
+/**
+ * Main program function.
+ */
+void RunUnrealSync(const TCHAR* CommandLine);
+
 /**
  * Helper class with functions used to sync engine.
  */
@@ -80,6 +85,16 @@ public:
 	 */
 	static void StartLoadingData();
 
+	/**
+	 * Tells that labels names are currently being loaded.
+	 *
+	 * @returns True if labels names are currently being loaded. False otherwise.
+	 */
+	static bool IsLoadingInProgress()
+	{
+		return LoaderThread.IsValid() && LoaderThread->IsInProgress();
+	}
+
 	/**
 	 * Terminates background P4 data loading process.
 	 */
diff --git a/Engine/Source/Programs/UnrealSync/Private/Windows/UnrealSyncMainWindows.cpp b/Engine/Source/Programs/UnrealSync/Private/Windows/UnrealSyncMainWindows.cpp
new file mode 100644
--- /dev/null
+++ b/Engine/Source/Programs/UnrealSync/Private/Windows/UnrealSyncMainWindows.cpp
@@ -0,0 +1,11 @@
+#include "../UnrealSync.h"
+
+/**
+ * WinMain, called when the application is started
+ */
+int WINAPI WinMain(_In_ HINSTANCE hInInstance, _In_opt_ HINSTANCE hPrevInstance, _In_ LPSTR, _In_ int nCmdShow)
+{
+	RunUnrealSync(GetCommandLineW());
+
+	return 0;
+}
\ No newline at end of file
