UnrealSync3 improvements.

The tool should be now built with the tools.

All params (if missing) are now auto-detected:
 - path to the p4.exe from default path or from default install location i.e. C:\\Program Files\\Perforce
 - p4 port and user are detected from P4V last connection string kept in <user>/.p4qt/ApplicationSettings.xml

diff --git a/Engine/Source/Programs/UnrealSync/Private/GUI.cpp b/Engine/Source/Programs/UnrealSync/Private/GUI.cpp
--- a/Engine/Source/Programs/UnrealSync/Private/GUI.cpp
+++ b/Engine/Source/Programs/UnrealSync/Private/GUI.cpp
@@ -88,7 +88,14 @@ private:
 	 */
 	void ComboBoxSelectionChanged(TSharedPtr<FString> Value, ESelectInfo::Type SelectInfo)
 	{
-		CurrentOption = OptionsSource.Find(Value);
+		CurrentOption = -1;
+		for (int32 OptionId = 0; OptionId < OptionsSource.Num(); ++OptionId)
+		{
+			if (OptionsSource[OptionId]->Equals(*Value))
+			{
+				CurrentOption = OptionId;
+			}
+		}
 
 		OnSelectionChanged.ExecuteIfBound(CurrentOption, SelectInfo);
 	}
@@ -514,7 +521,7 @@ public:
 	 */
 	FString GetLabelName() const OVERRIDE
 	{
-		return "-label=" + LabelsCombo->GetCurrentOption();
+		return LabelsCombo->GetCurrentOption();
 	}
 
 	/**
@@ -780,12 +787,7 @@ bool RunDetachedUS(const FString& USPath, bool bDoNotRunFromCopy, bool bDoNotUpd
 	FString CommandLine = FString()
 		+ (bDoNotRunFromCopy ? TEXT("-DoNotRunFromCopy ") : TEXT(""))
 		+ (bDoNotUpdateOnStartUp ? TEXT("-DoNotUpdateOnStartUp ") : TEXT(""))
-		+ (bPassP4Env
-		? FString::Printf(TEXT("-P4PORT=%s -P4USER=%s -US_FOUND_P4CLIENT=%s"),
-		*FP4Env::Get().GetPort(),
-		*FP4Env::Get().GetUser(),
-		*FP4Env::Get().GetClient())
-		: TEXT(""));
+		+ (bPassP4Env ? *FP4Env::Get().GetCommandLine()	: TEXT(""));
 
 	return RunProcess(USPath, CommandLine);
 }
@@ -978,7 +980,7 @@ private:
 
 		if (!FP4Env::Init(FCommandLine::Get()))
 		{
-			ErrorMsg = TEXT("P4 environment failed.");
+			ErrorMsg = TEXT("P4 environment detection failed.");
 			return false;
 		}
 
diff --git a/Engine/Source/Programs/UnrealSync/Private/P4Env.cpp b/Engine/Source/Programs/UnrealSync/Private/P4Env.cpp
--- a/Engine/Source/Programs/UnrealSync/Private/P4Env.cpp
+++ b/Engine/Source/Programs/UnrealSync/Private/P4Env.cpp
@@ -20,10 +20,23 @@ FString FP4Env::GetParam(const TCHAR* CommandLine, const FString& ParamName)
 	return Buf;
 }
 
-FP4Env::FP4Env(const FString& Port, const FString& User, const FString& Client, const FString& Branch)
-	: Port(Port), User(User), Client(Client), Branch(Branch)
+FP4Env::FP4Env(const TCHAR* CommandLine)
 {
+	class FParsingParamSerialization
+	{
+	public:
+		void SerializationTask(FString& FieldReference, EP4ParamType::Type Type)
+		{
+			FieldReference = GetParam(CommandLine, GetParamName(Type));
+		}
+
+		const TCHAR* CommandLine;
+	};
+
+	FParsingParamSerialization PS;
+	PS.CommandLine = CommandLine;
 
+	SerializeParams(FSerializationTask::CreateRaw(&PS, &FParsingParamSerialization::SerializationTask));
 }
 
 const FString& FP4Env::GetClient() const
@@ -72,22 +85,184 @@ FP4Env& FP4Env::Get()
 	return *Env;
 }
 
+/**
+ * Function that tries to detect P4 executable path.
+ *
+ * @param OutPath Output param. Found P4 executable path.
+ * @param Env P4 environment.
+ *
+ * @returns True if succeeded. False otherwise.
+ */
+bool DetectPath(FString& OutPath, const FP4Env& Env)
+{
+	// Tries to detect in standard environment paths.
+	
+	FString WhereCommand = "where"; // TODO Mac: I think 'where' command equivalent on Mac is 'whereis'
+	FString P4ExecutableName = "p4.exe";
+	FString Output;
+	
+	if (RunProcessOutput(WhereCommand, P4ExecutableName, Output))
+	{
+		// P4 found, quick exit.
+		OutPath = FPaths::ConvertRelativePathToFull(
+			Output.Replace(TEXT("\n"), TEXT("")).Replace(TEXT("\r"), TEXT("")));
+		return true;
+	}
+	
+	TArray<FString> PossibleLocations;
+
+	PossibleLocations.Add("C:\\Program Files\\Perforce");
+	PossibleLocations.Add("C:\\Program Files (x86)\\Perforce");
+
+	for (auto PossibleLocation : PossibleLocations)
+	{
+		FString LocationCandidate = FPaths::ConvertRelativePathToFull(FPaths::Combine(*PossibleLocation, *P4ExecutableName));
+		if (FPaths::FileExists(LocationCandidate))
+		{
+			OutPath = LocationCandidate;
+			return true;
+		}
+	}
+
+	return false;
+}
+
+#include "XmlParser.h"
+
+/**
+ * Tries to get last connection string from P4V config file.
+ *
+ * @param Output Output param.
+ * @param LastConnectionStringId Which string to output.
+ *
+ * @returns True if succeeded, false otherwise.
+ */
+bool GetP4VLastConnectionStringElement(FString& Output, int32 LastConnectionStringId)
+{
+	class FLastConnectionStringCache
+	{
+	public:
+		FLastConnectionStringCache()
+		{
+			bFoundData = false;
+			// TODO Mac: This path is going to be different on Mac.
+			FString AppSettingsXmlPath = FPaths::ConvertRelativePathToFull(
+					FPaths::Combine(FPlatformProcess::UserDir(),
+						TEXT(".."), TEXT(".p4qt"),
+						TEXT("ApplicationSettings.xml")
+					)
+				);
+
+			if (!FPaths::FileExists(AppSettingsXmlPath))
+			{
+				return;
+			}
+
+			TSharedPtr<FXmlFile> Doc = MakeShareable(new FXmlFile(AppSettingsXmlPath, EConstructMethod::ConstructFromFile));
+
+			for (FXmlNode* PropertyListNode : Doc->GetRootNode()->GetChildrenNodes())
+			{
+				if (PropertyListNode->GetTag() != "PropertyList" || PropertyListNode->GetAttribute("varName") != "Connection")
+				{
+					continue;
+				}
+
+				for (FXmlNode* VarNode : PropertyListNode->GetChildrenNodes())
+				{
+					if (VarNode->GetTag() != "String" || VarNode->GetAttribute("varName") != "LastConnection")
+					{
+						continue;
+					}
+
+					bFoundData = true;
+
+					FString Content = VarNode->GetContent();
+					FString Current;
+					FString Rest;
+
+					while (Content.Split(",", &Current, &Rest))
+					{
+						Current.Trim();
+						Current.TrimTrailing();
+
+						Data.Add(Current);
+
+						Content = Rest;
+					}
+
+					Rest.Trim();
+					Rest.TrimTrailing();
+
+					Data.Add(Rest);
+				}
+			}
+		}
+
+		bool bFoundData;
+		TArray<FString> Data;
+	};
+
+	static FLastConnectionStringCache Cache;
+
+	if (!Cache.bFoundData || Cache.Data.Num() <= LastConnectionStringId)
+	{
+		return false;
+	}
+
+	Output = Cache.Data[LastConnectionStringId];
+	return true;
+}
+
+/**
+ * Function that tries to detect P4 port.
+ *
+ * @param OutPort Output param. Found P4 port.
+ * @param Env P4 environment.
+ *
+ * @returns True if succeeded. False otherwise.
+ */
+bool DetectPort(FString& OutPort, const FP4Env& Env)
+{
+	if (!GetP4VLastConnectionStringElement(OutPort, 0))
+	{
+		return false;
+	}
+
+	return true;
+}
+
+/**
+ * Function that tries to detect P4 user name.
+ *
+ * @param OutUser Output param. Found P4 user name.
+ * @param Env P4 environment.
+ *
+ * @returns True if succeeded. False otherwise.
+ */
+bool DetectUser(FString& OutUser, const FP4Env& Env)
+{
+	if (!GetP4VLastConnectionStringElement(OutUser, 1))
+	{
+		return false;
+	}
+
+	return true;
+}
+
 #include "Internationalization/Regex.h"
 
 /**
- * Function that tries to detect P4 workspace name given port, user
- * using known path.
+ * Function that tries to detect P4 workspace name.
  *
  * @param OutClient Output param. Found P4 workspace name.
- * @param Port P4 port.
- * @param User P4 user.
+ * @param Env P4 environment.
  *
  * @returns True if succeeded. False otherwise.
  */
-bool DetectClient(FString& OutClient, const FString& Port, const FString& User)
+bool DetectClient(FString& OutClient, const FP4Env& Env)
 {
 	FString Output;
-	if (!RunProcessOutput(TEXT("p4.exe"), FString::Printf(TEXT("-p%s clients -u%s"), *Port, *User), Output))
+	if (!RunProcessOutput(Env.GetPath(), FString::Printf(TEXT("-p%s clients -u%s"), *Env.GetPort(), *Env.GetUser()), Output))
 	{
 		return false;
 	}
@@ -106,7 +281,7 @@ bool DetectClient(FString& OutClient, const FString& Port, const FString& User)
 		if (KnownPath.StartsWith(FPaths::ConvertRelativePathToFull(Root)))
 		{
 			FString InfoOutput;
-			if (!RunProcessOutput(TEXT("p4.exe"), FString::Printf(TEXT("-p%s -u%s -c%s info"), *Port, *User, *ClientName), InfoOutput))
+			if (!RunProcessOutput(Env.GetPath(), FString::Printf(TEXT("-p%s -u%s -c%s info"), *Env.GetPort(), *Env.GetUser(), *ClientName), InfoOutput))
 			{
 				return false;
 			}
@@ -132,20 +307,18 @@ bool DetectClient(FString& OutClient, const FString& Port, const FString& User)
 }
 
 /**
- * Tries to detect branch from P4 port, user and workspace name
- * using known depot path.
+ * Tries to detect branch using known depot path.
  *
  * @param OutBranch Output parameter. Detected branch.
- * @param Port P4 port.
- * @param User P4 user.
- * @param Client P4 workspace name.
+ * @param Env P4 environment.
  *
  * @returns True if succeeded. False otherwise.
  */
-bool DetectBranch(FString& OutBranch, const FString& Port, const FString& User, const FString& Client)
+bool DetectBranch(FString& OutBranch, const FP4Env& Env)
 {
 	FString Output;
-	if (!RunProcessOutput(TEXT("p4.exe"), FString::Printf(TEXT("-p%s -u%s -c%s files %s"), *Port, *User, *Client, *GetKnownPath()), Output))
+	if (!RunProcessOutput(Env.GetPath(), FString::Printf(TEXT("-p%s -u%s -c%s files %s"),
+		*Env.GetPort(), *Env.GetUser(), *Env.GetClient(), *GetKnownPath()), Output))
 	{
 		return false;
 	}
@@ -160,41 +333,23 @@ bool DetectBranch(FString& OutBranch, const FString& Port, const FString& User,
 	return true;
 }
 
-bool FP4Env::Init(const FString& Port, const FString& User, const FString& Client)
+bool FP4Env::Init(const TCHAR* CommandLine)
 {
-	FString TmpClient;
-	FString Branch;
-
-	if (Client.IsEmpty())
-	{
-		if (!DetectClient(TmpClient, Port, User))
-		{
-			return false;
-		}
-	}
-	else
-	{
-		TmpClient = Client;
-	}
+	TSharedPtr<FP4Env> Env = MakeShareable(new FP4Env(CommandLine));
 
-	if (!DetectBranch(Branch, Port, User, TmpClient))
+	if (!Env->AutoDetectMissingParams())
 	{
 		return false;
 	}
 
-	Env = MakeShareable(new FP4Env(Port, User, TmpClient, Branch));
-	return true;
-}
+	FP4Env::Env = Env;
 
-bool FP4Env::Init(const TCHAR* CommandLine)
-{
-	// Auto-detect P4.
-	return Init(GetParam(CommandLine, TEXT("P4PORT")), GetParam(CommandLine, TEXT("P4USER")), GetParam(CommandLine, TEXT("US_FOUND_P4CLIENT")));
+	return true;
 }
 
 bool FP4Env::RunP4Progress(const FString& CommandLine, const FOnP4MadeProgress& OnP4MadeProgress)
 {
-	return RunProcessProgress(TEXT("p4.exe"), FString::Printf(TEXT("-p%s -c%s -u%s %s"), *Get().GetPort(), *Get().GetClient(), *Get().GetUser(), *CommandLine), OnP4MadeProgress);
+	return RunProcessProgress(Get().GetPath(), FString::Printf(TEXT("-p%s -c%s -u%s %s"), *Get().GetPort(), *Get().GetClient(), *Get().GetUser(), *CommandLine), OnP4MadeProgress);
 }
 
 bool FP4Env::RunP4Output(const FString& CommandLine, FString& Output)
@@ -228,5 +383,114 @@ bool FP4Env::RunP4(const FString& CommandLine)
 	return RunP4Progress(CommandLine, nullptr);
 }
 
+void FP4Env::SerializeParams(const FSerializationTask& SerializationTask)
+{
+	SerializationTask.ExecuteIfBound(Path, EP4ParamType::Path);
+	SerializationTask.ExecuteIfBound(Port, EP4ParamType::Port);
+	SerializationTask.ExecuteIfBound(User, EP4ParamType::User);
+	SerializationTask.ExecuteIfBound(Client, EP4ParamType::Client);
+	SerializationTask.ExecuteIfBound(Branch, EP4ParamType::Branch);
+}
+
+FString FP4Env::GetParamName(EP4ParamType::Type Type)
+{
+	switch (Type)
+	{
+	case EP4ParamType::Path:
+		return "P4PATH";
+	case EP4ParamType::Port:
+		return "P4PORT";
+	case EP4ParamType::User:
+		return "P4USER";
+	case EP4ParamType::Client:
+		return "P4CLIENT";
+	case EP4ParamType::Branch:
+		return "P4BRANCH";
+	}
+
+	checkNoEntry();
+	return "";
+}
+
+bool FP4Env::AutoDetectMissingParams()
+{
+	class FAutoDetectTasks
+	{
+	public:
+		void AutoDetectParam(FString& FieldReference, EP4ParamType::Type Type)
+		{
+			// If parameter already found then stop.
+			if (!FieldReference.IsEmpty())
+			{
+				return;
+			}
+
+			// If previous detection failed then stop.
+			if (Stop)
+			{
+				return;
+			}
+
+			switch (Type)
+			{
+			case EP4ParamType::Path:
+				Stop = !DetectPath(FieldReference, Env);
+				break;
+			case EP4ParamType::Port:
+				Stop = !DetectPort(FieldReference, Env);
+				break;
+			case EP4ParamType::User:
+				Stop = !DetectUser(FieldReference, Env);
+				break;
+			case EP4ParamType::Client:
+				Stop = !DetectClient(FieldReference, Env);
+				break;
+			case EP4ParamType::Branch:
+				Stop = !DetectBranch(FieldReference, Env);
+				break;
+			}
+		}
+
+		FAutoDetectTasks(FP4Env& Env)
+			: Env(Env)
+		{
+			Stop = false;
+		}
+
+		FP4Env& Env;
+		bool Stop;
+	};
+
+	FAutoDetectTasks AT(*this);
+	SerializeParams(FSerializationTask::CreateRaw(&AT, &FAutoDetectTasks::AutoDetectParam));
+
+	return !AT.Stop;
+}
+
+FString FP4Env::GetCommandLine()
+{
+	class CommandLineOutput
+	{
+	public:
+		void AppendParam(FString& FieldReference, EP4ParamType::Type Type)
+		{
+			Output += " -" + GetParamName(Type) + "=\"" + FieldReference + "\"";
+		}
+
+		FString Output;
+	};
+
+	CommandLineOutput Output;
+
+	SerializeParams(FSerializationTask::CreateRaw(&Output, &CommandLineOutput::AppendParam));
+
+	return Output.Output;
+}
+
+const FString& FP4Env::GetPath() const
+{
+	return Path;
+}
+
 /* Static variable initialization. */
 TSharedPtr<FP4Env> FP4Env::Env;
\ No newline at end of file
diff --git a/Engine/Source/Programs/UnrealSync/Private/P4Env.h b/Engine/Source/Programs/UnrealSync/Private/P4Env.h
--- a/Engine/Source/Programs/UnrealSync/Private/P4Env.h
+++ b/Engine/Source/Programs/UnrealSync/Private/P4Env.h
@@ -5,6 +5,21 @@
 #include "Core.h"
 #include "CoreUObject.h"
 
+/**
+ * Enum describing P4 param type.
+ */
+namespace EP4ParamType
+{
+	enum Type
+	{
+		Path,
+		Port,
+		User,
+		Client,
+		Branch
+	};
+}
+
 /**
  * P4 environment class to connect with P4 executable.
  */
@@ -73,6 +88,13 @@ public:
 
 
 
+	/**
+	 * Gets P4 path.
+	 *
+	 * @returns P4 path.
+	 */
+	const FString& GetPath() const;
+
 	/**
 	 * Gets P4 user.
 	 *
@@ -101,16 +123,46 @@ public:
 	 */
 	const FString& GetBranch() const;
 
+	/**
+	 * Gets command line string to pass to new UnrealSync process to copy current environment.
+	 *
+	 * @returns Command line string.
+	 */
+	FString GetCommandLine();
+
 private:
+	/* Param serialization delegate. */
+	DECLARE_DELEGATE_TwoParams(FSerializationTask, FString&, EP4ParamType::Type);
+
 	/**
 	 * Constructor
 	 *
-	 * @param Port P4 port.
-	 * @param User P4 user.
-	 * @param Client P4 workspace name.
-	 * @param Branch Current P4 branch.
+	 * @param CommandLine Command line to parse params from.
 	 */
-	FP4Env(const FString& Port, const FString& User, const FString& Client, const FString& Branch);
+	FP4Env(const TCHAR* CommandLine);
+
+	/**
+	 * Serializes all P4Env params.
+	 *
+	 * @param SerializationTask Task to perform serialization.
+	 */
+	void SerializeParams(const FSerializationTask& SerializationTask);
+
+	/**
+	 * Auto-detect missing params.
+	 *
+	 * @returns True if succeeded. False otherwise.
+	 */
+	bool AutoDetectMissingParams();
+
+	/**
+	 * Gets param type name.
+	 *
+	 * @param Type Param type enum id.
+	 *
+	 * @returns Param type name.
+	 */
+	static FString GetParamName(EP4ParamType::Type Type);
 
 	/**
 	 * Parses param either from command line or environment variables.
@@ -126,6 +178,9 @@ private:
 	/* Singleton instance pointer. */
 	static TSharedPtr<FP4Env> Env;
 
+	/* Path to P4 executable. */
+	FString Path;
+
 	/* P4 port. */
 	FString Port;
 
diff --git a/Engine/Source/Programs/UnrealSync/UnrealSync.Target.cs b/Engine/Source/Programs/UnrealSync/UnrealSync.Target.cs
--- a/Engine/Source/Programs/UnrealSync/UnrealSync.Target.cs
+++ b/Engine/Source/Programs/UnrealSync/UnrealSync.Target.cs
@@ -55,4 +55,11 @@ public class UnrealSyncTarget : TargetRules
 
 		OutLinkEnvironmentConfiguration.bIsBuildingConsoleApplication = false;
 	}
+
+	public override bool GUBP_AlwaysBuildWithTools(UnrealTargetPlatform InHostPlatform, out bool bInternalToolOnly, out bool SeparateNode)
+	{
+		bInternalToolOnly = true;
+		SeparateNode = false;
+		return true;
+	}
 }
