CIS fix. Allow building of UnrealSync3 only on Windows platforms as it's not supported yet on the others.

diff --git a/Engine/Source/Programs/UnrealSync/UnrealSync.Target.cs b/Engine/Source/Programs/UnrealSync/UnrealSync.Target.cs
--- a/Engine/Source/Programs/UnrealSync/UnrealSync.Target.cs
+++ b/Engine/Source/Programs/UnrealSync/UnrealSync.Target.cs
@@ -58,8 +58,13 @@ public class UnrealSyncTarget : TargetRules
 
 	public override bool GUBP_AlwaysBuildWithTools(UnrealTargetPlatform InHostPlatform, out bool bInternalToolOnly, out bool SeparateNode)
 	{
-		bInternalToolOnly = true;
-		SeparateNode = false;
-		return true;
+		if (InHostPlatform == UnrealTargetPlatform.Win32 || InHostPlatform == UnrealTargetPlatform.Win64)
+		{
+			bInternalToolOnly = true;
+			SeparateNode = false;
+			return true;
+		}
+
+		return false;
 	}
 }
