CIS fix.

diff --git a/Engine/Source/Programs/UnrealSync/UnrealSync.Target.cs b/Engine/Source/Programs/UnrealSync/UnrealSync.Target.cs
--- a/Engine/Source/Programs/UnrealSync/UnrealSync.Target.cs
+++ b/Engine/Source/Programs/UnrealSync/UnrealSync.Target.cs
@@ -65,6 +65,8 @@ public class UnrealSyncTarget : TargetRules
 			return true;
 		}
 
+		bInternalToolOnly = false;
+		SeparateNode = false;
 		return false;
 	}
 }
