UnrealSync3: Now P4PORT is falling back to default perforce:1666 and P4USER is read from p4 info output if previous sources fail.

diff --git a/Engine/Source/Programs/UnrealSync/Private/P4Env.cpp b/Engine/Source/Programs/UnrealSync/Private/P4Env.cpp
--- a/Engine/Source/Programs/UnrealSync/Private/P4Env.cpp
+++ b/Engine/Source/Programs/UnrealSync/Private/P4Env.cpp
@@ -223,11 +223,13 @@ bool GetP4VLastConnectionStringElement(FString& Output, int32 LastConnectionStri
  */
 bool DetectPort(FString& OutPort, const FP4Env& Env)
 {
-	if (!GetP4VLastConnectionStringElement(OutPort, 0))
+	if (GetP4VLastConnectionStringElement(OutPort, 0))
 	{
-		return false;
+		return true;
 	}
 
+	// Fallback to default port. If it's not valid auto-detection will fail later.
+	OutPort = "perforce:1666";
 	return true;
 }
 
@@ -241,12 +243,24 @@ bool DetectPort(FString& OutPort, const FP4Env& Env)
  */
 bool DetectUser(FString& OutUser, const FP4Env& Env)
 {
-	if (!GetP4VLastConnectionStringElement(OutUser, 1))
+	if (GetP4VLastConnectionStringElement(OutUser, 1))
 	{
-		return false;
+		return true;
 	}
 
-	return true;
+	FString Output;
+	RunProcessOutput(Env.GetPath(), FString::Printf(TEXT("-p%s info"), *Env.GetPort()), Output);
+
+	const FRegexPattern UserNamePattern(TEXT("User name:\\s*([^ \\t\\n\\r]+)\\s*"));
+
+	FRegexMatcher Matcher(UserNamePattern, Output);
+	if (Matcher.FindNext())
+	{
+		OutUser = Output.Mid(Matcher.GetCaptureGroupBeginning(1), Matcher.GetCaptureGroupEnding(1) - Matcher.GetCaptureGroupBeginning(1));
+		return true;
+	}
+
+	return false;
 }
 
 #include "Internationalization/Regex.h"
