UnrealSync: Improved the way P4 environment is detected. The program now tries all possible detection configurations before fail instead of sticking to first found possibility for every parameter.

diff --git a/Engine/Source/Programs/UnrealSync/Private/P4Env.cpp b/Engine/Source/Programs/UnrealSync/Private/P4Env.cpp
--- a/Engine/Source/Programs/UnrealSync/Private/P4Env.cpp
+++ b/Engine/Source/Programs/UnrealSync/Private/P4Env.cpp
@@ -6,38 +6,45 @@
 #include "ProcessHelper.h"
 #include "Regex.h"
 
-FString FP4Env::GetParam(const TCHAR* CommandLine, const FString& ParamName)
+/**
+ * Parses param from command line.
+ *
+ * @param Value Output value.
+ * @param CommandLine Command line to parse.
+ * @param ParamName Param name to parse.
+ *
+ * @returns True if found, false otherwise.
+ */
+bool GetParam(FString& Value, const TCHAR* CommandLine, const FString& ParamName)
 {
-	FString Value;
-	if (FParse::Value(CommandLine, *(ParamName + "="), Value))
-	{
-		return Value;
-	}
+	return FParse::Value(CommandLine, *(ParamName + "="), Value);
+}
 
+/**
+ * Gets param from environment variables.
+ *
+ * @param Value Output value.
+ * @param ParamName Param name to parse.
+ *
+ * @returns True if found, false otherwise.
+ */
+bool GetEnvParam(FString& Value, const FString& ParamName)
+{
 	TCHAR Buf[512];
 
 	FPlatformMisc::GetEnvironmentVariable(*ParamName, Buf, 512);
 
-	return Buf;
-}
-
-FP4Env::FP4Env(const TCHAR* CommandLine)
-{
-	class FParsingParamSerialization
+	if (FCString::Strlen(Buf) > 0)
 	{
-	public:
-		void SerializationTask(FString& FieldReference, EP4ParamType::Type Type)
-		{
-			FieldReference = GetParam(CommandLine, GetParamName(Type));
-		}
-
-		const TCHAR* CommandLine;
-	};
+		Value = Buf;
+		return true;
+	}
 
-	FParsingParamSerialization PS;
-	PS.CommandLine = CommandLine;
+	return false;
+}
 
-	SerializeParams(FSerializationTask::CreateRaw(&PS, &FParsingParamSerialization::SerializationTask));
+FP4Env::FP4Env()
+{
 }
 
 const FString& FP4Env::GetClient() const
@@ -87,46 +94,211 @@ FP4Env& FP4Env::Get()
 }
 
 /**
- * Function that tries to detect P4 executable path.
- *
- * @param OutPath Output param. Found P4 executable path.
- * @param Env P4 environment.
- *
- * @returns True if succeeded. False otherwise.
+ * Param detection iterator interface.
  */
-bool DetectPath(FString& OutPath, const FP4Env& Env)
+class IP4EnvParamDetectionIterator
 {
-	// Tries to detect in standard environment paths.
-	
-	FString WhereCommand = "where"; // TODO Mac: I think 'where' command equivalent on Mac is 'whereis'
-	FString P4ExecutableName = "p4.exe";
-	FString Output;
-	
-	if (RunProcessOutput(WhereCommand, P4ExecutableName, Output))
+public:
+	/**
+	 * Try to auto-detect next param proposition.
+	 *
+	 * @returns True if found. False otherwise.
+	 */
+	virtual bool MoveNext() = 0;
+
+	/**
+	 * Gets currently found param proposition.
+	 *
+	 * @returns Found param.
+	 */
+	virtual const FString& GetCurrent() const = 0;
+
+	/**
+	 * Creates auto-detection iterator for given type.
+	 *
+	 * @param Type Type of iterator to create.
+	 * @param CommandLine Command line to parse by iterator.
+	 * @param Env Current P4 environment state.
+	 *
+	 * @returns Shared pointer to created iterator.
+	 */
+	static TSharedPtr<IP4EnvParamDetectionIterator> Create(EP4ParamType::Type Type, const TCHAR *CommandLine, const FP4Env& Env);
+};
+
+/**
+ * Base class for param detection iterators.
+ */
+class FP4EnvParamDetectionIteratorBase : public IP4EnvParamDetectionIterator
+{
+public:
+	/**
+	 * Constructor
+	 *
+	 * @param Type Type of the param to auto-detect.
+	 * @param CommandLine Command line to check for param.
+	 */
+	FP4EnvParamDetectionIteratorBase(EP4ParamType::Type Type, const TCHAR* CommandLine)
+		: Type(Type)
 	{
-		// P4 found, quick exit.
-		OutPath = FPaths::ConvertRelativePathToFull(
-			Output.Replace(TEXT("\n"), TEXT("")).Replace(TEXT("\r"), TEXT("")));
-		return true;
+		if (GetParam(CommandLineParam, CommandLine, FP4Env::GetParamName(Type)))
+		{
+			Finish();
+		}
 	}
-	
-	TArray<FString> PossibleLocations;
 
-	PossibleLocations.Add("C:\\Program Files\\Perforce");
-	PossibleLocations.Add("C:\\Program Files (x86)\\Perforce");
+	/**
+	 * Try to auto-detect next param proposition.
+	 *
+	 * @returns True if found. False otherwise.
+	 */
+	virtual bool MoveNext() override
+	{
+		if (bFinished)
+		{
+			if (!CommandLineParam.IsEmpty())
+			{
+				Current = CommandLineParam;
+				CommandLineParam.Empty();
+				return true;
+			}
+			else
+			{
+				return false;
+			}
+		}
+
+		if (!bEnvChecked)
+		{
+			bEnvChecked = true;
+
+			if (GetEnvParam(Current, FP4Env::GetParamName(Type)))
+			{
+				return true;
+			}
+		}
+
+		return FindNext(Current);
+	}
+
+	/**
+	 * Gets currently found param proposition.
+	 *
+	 * @returns Found param.
+	 */
+	virtual const FString& GetCurrent() const override
+	{
+		return Current;
+	}
+
+protected:
+	/**
+	 * Find next param.
+	 *
+	 * @param Output Output parameter.
+	 *
+	 * @returns True if found. False otherwise.
+	 */
+	virtual bool FindNext(FString& Output) = 0;
+
+	/**
+	 * Mark this iterator as finished.
+	 */
+	void Finish()
+	{
+		bFinished = true;
+	}
+
+private:
+	/** Currently found param proposition. */
+	FString Current;
+
+	/** Found command line param. */
+	FString CommandLineParam;
+
+	/** Type of the param. */
+	EP4ParamType::Type Type;
+
+	/** Is this iterator finished. */
+	bool bFinished = false;
 
-	for (auto PossibleLocation : PossibleLocations)
+	/** Tells if environment variable have been checked already by this iterator. */
+	bool bEnvChecked = false;
+};
+
+/**
+ * P4 path param auto-detection iterator.
+ */
+class FP4PathDetectionIterator : public FP4EnvParamDetectionIteratorBase
+{
+public:
+	/**
+	 * Constructor.
+	 *
+	 * @param CommandLine Command line to parse.
+	 */
+	FP4PathDetectionIterator(const TCHAR* CommandLine)
+		: FP4EnvParamDetectionIteratorBase(EP4ParamType::Path, CommandLine)
+	{
+		PossibleLocations.Add("C:\\Program Files\\Perforce");
+		PossibleLocations.Add("C:\\Program Files (x86)\\Perforce");
+	}
+
+	/**
+	 * Function that tries to detect P4 executable path.
+	 *
+	 * @param Output Output param. Found P4 executable path.
+	 *
+	 * @returns True if succeeded. False otherwise.
+	 */
+	virtual bool FindNext(FString& Output) override
 	{
-		FString LocationCandidate = FPaths::ConvertRelativePathToFull(FPaths::Combine(*PossibleLocation, *P4ExecutableName));
-		if (FPaths::FileExists(LocationCandidate))
+		// Tries to detect in standard environment paths.
+		static const FString WhereCommand = "where"; // TODO Mac: I think 'where' command equivalent on Mac is 'whereis'
+		static const FString P4ExecutableName = "p4.exe";
+
+		if (Step == 0)
 		{
-			OutPath = LocationCandidate;
-			return true;
+			++Step;
+
+			FString WhereOutput;
+			if (RunProcessOutput(WhereCommand, P4ExecutableName, WhereOutput))
+			{
+				// P4 found, quick exit.
+				Output = FPaths::ConvertRelativePathToFull(
+					WhereOutput.Replace(TEXT("\n"), TEXT("")).Replace(TEXT("\r"), TEXT("")));
+				return true;
+			}
+
+			return false;
 		}
+
+		if (Step > PossibleLocations.Num())
+		{
+			Finish();
+			return false;
+		}
+
+		for (; Step <= PossibleLocations.Num(); ++Step)
+		{
+			FString LocationCandidate = FPaths::ConvertRelativePathToFull(FPaths::Combine(*PossibleLocations[Step - 1], *P4ExecutableName));
+			if (FPaths::FileExists(LocationCandidate))
+			{
+				Output = LocationCandidate;
+				return true;
+			}
+		}
+
+		Finish();
+		return false;
 	}
 
-	return false;
-}
+private:
+	/** Possible P4 installation locations. */
+	TArray<FString> PossibleLocations;
+
+	/** Current step number. */
+	int32 Step;
+};
 
 #include "XmlParser.h"
 
@@ -215,132 +387,224 @@ bool GetP4VLastConnectionStringElement(FString& Output, int32 LastConnectionStri
 }
 
 /**
- * Function that tries to detect P4 port.
- *
- * @param OutPort Output param. Found P4 port.
- * @param Env P4 environment.
- *
- * @returns True if succeeded. False otherwise.
+ * P4 port param auto-detection iterator.
  */
-bool DetectPort(FString& OutPort, const FP4Env& Env)
+class FP4PortDetectionIterator : public FP4EnvParamDetectionIteratorBase
 {
-	if (GetP4VLastConnectionStringElement(OutPort, 0))
+public:
+	/**
+	 * Constructor.
+	 *
+	 * @param CommandLine Command line to parse.
+	 */
+	FP4PortDetectionIterator(const TCHAR* CommandLine)
+		: FP4EnvParamDetectionIteratorBase(EP4ParamType::Port, CommandLine)
+	{ }
+
+	/**
+	 * Function that tries to detect P4 port.
+	 *
+	 * @param Output Output param. Found P4 port.
+	 *
+	 * @returns True if succeeded. False otherwise.
+	 */
+	virtual bool FindNext(FString& Output) override
 	{
-		return true;
+		if (Step == 0)
+		{
+			++Step;
+			if (GetP4VLastConnectionStringElement(Output, 0))
+			{
+				return true;
+			}
+		}
+
+		if (Step == 1)
+		{
+			++Step;
+
+			// Fallback to default port. If it's not valid auto-detection will fail later.
+			Output = "perforce:1666";
+			return true;
+		}
+
+		Finish();
+		return false;
 	}
 
-	// Fallback to default port. If it's not valid auto-detection will fail later.
-	OutPort = "perforce:1666";
-	return true;
-}
+private:
+	/** Current step number. */
+	int32 Step;
+};
 
 /**
- * Function that tries to detect P4 user name.
- *
- * @param OutUser Output param. Found P4 user name.
- * @param Env P4 environment.
- *
- * @returns True if succeeded. False otherwise.
+ * P4 user param auto-detection iterator.
  */
-bool DetectUser(FString& OutUser, const FP4Env& Env)
+class FP4UserDetectionIterator : public FP4EnvParamDetectionIteratorBase
 {
-	if (GetP4VLastConnectionStringElement(OutUser, 1))
+public:
+	/**
+	 * Constructor.
+	 *
+	 * @param CommandLine Command line to parse.
+	 * @param Env Current P4 environment state.
+	 */
+	FP4UserDetectionIterator(const TCHAR* CommandLine, const FP4Env& Env)
+		: FP4EnvParamDetectionIteratorBase(EP4ParamType::User, CommandLine), Env(Env)
 	{
-		return true;
 	}
 
-	FString Output;
-	RunProcessOutput(Env.GetPath(), FString::Printf(TEXT("-p%s info"), *Env.GetPort()), Output);
+	/**
+	 * Function that tries to detect P4 user.
+	 *
+	 * @param Output Output param. Found P4 user.
+	 *
+	 * @returns True if succeeded. False otherwise.
+	 */
+	virtual bool FindNext(FString& Output) override
+	{
+		if (Step == 0)
+		{
+			++Step;
+			if (GetP4VLastConnectionStringElement(Output, 1))
+			{
+				return true;
+			}
+		}
+
+		if (Step == 1)
+		{
+			++Step;
 
-	const FRegexPattern UserNamePattern(TEXT("User name:\\s*([^ \\t\\n\\r]+)\\s*"));
+			FString InfoOutput;
+			RunProcessOutput(Env.GetPath(), FString::Printf(TEXT("-p%s info"), *Env.GetPort()), InfoOutput);
 
-	FRegexMatcher Matcher(UserNamePattern, Output);
-	if (Matcher.FindNext())
-	{
-		OutUser = Output.Mid(Matcher.GetCaptureGroupBeginning(1), Matcher.GetCaptureGroupEnding(1) - Matcher.GetCaptureGroupBeginning(1));
-		return true;
+			const FRegexPattern UserNamePattern(TEXT("User name:\\s*([^ \\t\\n\\r]+)\\s*"));
+
+			FRegexMatcher Matcher(UserNamePattern, InfoOutput);
+			if (Matcher.FindNext())
+			{
+				Output = InfoOutput.Mid(Matcher.GetCaptureGroupBeginning(1), Matcher.GetCaptureGroupEnding(1) - Matcher.GetCaptureGroupBeginning(1));
+				return true;
+			}
+		}
+
+		Finish();
+		return false;
 	}
 
-	return false;
-}
+private:
+	/** Current P4 environment state. */
+	const FP4Env& Env;
+	
+	/** Current step number. */
+	int32 Step;
+};
 
 #include "Internationalization/Regex.h"
 
 /**
- * Function that tries to detect P4 workspace name.
- *
- * @param OutClient Output param. Found P4 workspace name.
- * @param Env P4 environment.
- *
- * @returns True if succeeded. False otherwise.
+ * P4 client param auto-detection iterator.
  */
-bool DetectClient(FString& OutClient, const FP4Env& Env)
+class FP4ClientDetectionIterator : public FP4EnvParamDetectionIteratorBase
 {
-	FString Output;
-	if (!RunProcessOutput(Env.GetPath(), FString::Printf(TEXT("-p%s clients -u%s"), *Env.GetPort(), *Env.GetUser()), Output))
+public:
+	/**
+	 * Constructor.
+	 *
+	 * @param CommandLine Command line to parse.
+	 * @param Env Current P4 environment state.
+	 */
+	FP4ClientDetectionIterator(const TCHAR* CommandLine, const FP4Env& Env)
+		: FP4EnvParamDetectionIteratorBase(EP4ParamType::Client, CommandLine), Env(Env)
 	{
-		return false;
-	}
-
-	FString KnownPath = GetKnownPath();
-	FString HostName = GetHostName();
+		if (!RunProcessOutput(Env.GetPath(), FString::Printf(TEXT("-p%s clients -u%s"), *Env.GetPort(), *Env.GetUser()), P4ClientsOutput))
+		{
+			Finish();
+			return;
+		}
 
-	const FRegexPattern ClientsPattern(TEXT("Client ([^ ]+) \\d{4}/\\d{2}/\\d{2} root (.+) '.*'"));
-	FRegexMatcher Matcher(ClientsPattern, Output);
+		static const FRegexPattern ClientsPattern(TEXT("Client ([^ ]+) \\d{4}/\\d{2}/\\d{2} root (.+) '.*'"));
+		Matcher = MakeShareable(new FRegexMatcher(ClientsPattern, P4ClientsOutput));
+	}
 
-	while (Matcher.FindNext())
+	/**
+	 * Function that tries to detect P4 client.
+	 *
+	 * @param Output Output param. Found P4 client.
+	 *
+	 * @returns True if succeeded. False otherwise.
+	 */
+	virtual bool FindNext(FString& Output) override
 	{
-		auto ClientName = Output.Mid(Matcher.GetCaptureGroupBeginning(1), Matcher.GetCaptureGroupEnding(1) - Matcher.GetCaptureGroupBeginning(1));
-		auto Root = Output.Mid(Matcher.GetCaptureGroupBeginning(2), Matcher.GetCaptureGroupEnding(2) - Matcher.GetCaptureGroupBeginning(2));
+		static const FString KnownPath = GetKnownPath();
+		static const FString HostName = GetHostName();
 
-		if (KnownPath.StartsWith(FPaths::ConvertRelativePathToFull(Root)))
+		while (Matcher->FindNext())
 		{
-			FString InfoOutput;
-			if (!RunProcessOutput(Env.GetPath(), FString::Printf(TEXT("-p%s -u%s -c%s info"), *Env.GetPort(), *Env.GetUser(), *ClientName), InfoOutput))
+			auto ClientName = P4ClientsOutput.Mid(Matcher->GetCaptureGroupBeginning(1), Matcher->GetCaptureGroupEnding(1) - Matcher->GetCaptureGroupBeginning(1));
+			auto Root = P4ClientsOutput.Mid(Matcher->GetCaptureGroupBeginning(2), Matcher->GetCaptureGroupEnding(2) - Matcher->GetCaptureGroupBeginning(2));
+
+			if (KnownPath.StartsWith(FPaths::ConvertRelativePathToFull(Root)))
 			{
-				return false;
-			}
+				FString InfoOutput;
+				if (!RunProcessOutput(Env.GetPath(), FString::Printf(TEXT("-p%s -u%s -c%s info"), *Env.GetPort(), *Env.GetUser(), *ClientName), InfoOutput))
+				{
+					continue;
+				}
 
-			const FRegexPattern InfoPattern(TEXT("Client host: ([^\\r\\n ]+)\\s"));
-			FRegexMatcher InfoMatcher(InfoPattern, InfoOutput);
+				const FRegexPattern InfoPattern(TEXT("Client host: ([^\\r\\n ]+)\\s"));
+				FRegexMatcher InfoMatcher(InfoPattern, InfoOutput);
 
-			if (InfoMatcher.FindNext())
-			{
-				if (InfoOutput.Mid(
-					InfoMatcher.GetCaptureGroupBeginning(1),
-					InfoMatcher.GetCaptureGroupEnding(1) - InfoMatcher.GetCaptureGroupBeginning(1)
-					).Equals(HostName))
+				if (InfoMatcher.FindNext())
 				{
-					OutClient = ClientName;
-					return true;
+					if (InfoOutput.Mid(
+						InfoMatcher.GetCaptureGroupBeginning(1),
+						InfoMatcher.GetCaptureGroupEnding(1) - InfoMatcher.GetCaptureGroupBeginning(1)
+						).Equals(HostName))
+					{
+						Output = ClientName;
+						return true;
+					}
 				}
 			}
 		}
+
+		Finish();
+		return false;
 	}
 
-	return false;
-}
+private:
+	/** Current regex matcher that parses next client. */
+	TSharedPtr<FRegexMatcher> Matcher;
+
+	/** Current P4 environment state. */
+	const FP4Env& Env;
+
+	/** Output of the p4 clients command. */
+	FString P4ClientsOutput;
+};
 
 /**
- * Tries to detect branch using known depot path.
+ * Gets branch detected for this app from current P4 environment state.
  *
- * @param OutBranch Output parameter. Detected branch.
- * @param Env P4 environment.
+ * @param Output Found branch prefix.
+ * @parma Env Current P4 environment state.
  *
- * @returns True if succeeded. False otherwise.
+ * @returns True if found, false otherwise.
  */
-bool DetectBranch(FString& OutBranch, const FP4Env& Env)
+bool GetCurrentBranch(FString& Output, const FP4Env& Env)
 {
-	FString Output;
+	FString FilesOutput;
 	if (!RunProcessOutput(Env.GetPath(), FString::Printf(TEXT("-p%s -u%s -c%s files %s"),
-		*Env.GetPort(), *Env.GetUser(), *Env.GetClient(), *GetKnownPath()), Output))
+		*Env.GetPort(), *Env.GetUser(), *Env.GetClient(), *GetKnownPath()), FilesOutput))
 	{
 		return false;
 	}
 
 	FString Rest;
 
-	if (!Output.Split("/Engine/", &OutBranch, &Rest))
+	if (!FilesOutput.Split("/Engine/", &Output, &Rest))
 	{
 		return false;
 	}
@@ -348,11 +612,53 @@ bool DetectBranch(FString& OutBranch, const FP4Env& Env)
 	return true;
 }
 
+/**
+ * P4 client param auto-detection iterator.
+ */
+class FP4BranchDetectionIterator : public FP4EnvParamDetectionIteratorBase
+{
+public:
+	/**
+	 * Constructor.
+	 *
+	 * @param CommandLine Command line to parse.
+	 * @param Env Current P4 environment state.
+	 */
+	FP4BranchDetectionIterator(const TCHAR* CommandLine, const FP4Env& Env)
+		: FP4EnvParamDetectionIteratorBase(EP4ParamType::Branch, CommandLine), Env(Env)
+	{
+	}
+
+	/**
+	 * Function that tries to detect P4 branch.
+	 *
+	 * @param Output Output param. Found P4 branch.
+	 *
+	 * @returns True if succeeded. False otherwise.
+	 */
+	virtual bool FindNext(FString& Output) override
+	{
+		Finish();
+
+		return GetCurrentBranch(Output, Env);
+	}
+
+private:
+	/** Current P4 environment state. */
+	const FP4Env& Env;
+};
+
 bool FP4Env::Init(const TCHAR* CommandLine)
 {
-	TSharedPtr<FP4Env> Env = MakeShareable(new FP4Env(CommandLine));
+	TSharedPtr<FP4Env> Env = MakeShareable(new FP4Env());
+
+	if (!Env->AutoDetectMissingParams(CommandLine))
+	{
+		return false;
+	}
 
-	if (!Env->AutoDetectMissingParams())
+	FString CurrentBranch;
+	if (!GetCurrentBranch(CurrentBranch, *Env) || CurrentBranch != Env->GetBranch())
 	{
 		return false;
 	}
@@ -427,59 +733,38 @@ FString FP4Env::GetParamName(EP4ParamType::Type Type)
 	return "";
 }
 
-bool FP4Env::AutoDetectMissingParams()
+bool FP4Env::AutoDetectMissingParams(const TCHAR* CommandLine)
 {
-	class FAutoDetectTasks
+	TArray<TSharedPtr<IP4EnvParamDetectionIterator> > ParamDetectionIteratorsStack;
+
+	EP4ParamType::Type Type = EP4ParamType::Path;
+	ParamDetectionIteratorsStack.Add(IP4EnvParamDetectionIterator::Create(Type, CommandLine, *this));
+
+	while (ParamDetectionIteratorsStack.Num() > 0)
 	{
-	public:
-		void AutoDetectParam(FString& FieldReference, EP4ParamType::Type Type)
+		if (ParamDetectionIteratorsStack.Last()->MoveNext())
 		{
-			// If parameter already found then stop.
-			if (!FieldReference.IsEmpty())
-			{
-				return;
-			}
+			SetParam(Type, ParamDetectionIteratorsStack.Last()->GetCurrent());
 
-			// If previous detection failed then stop.
-			if (Stop)
+			if (Type != EP4ParamType::Branch)
 			{
-				return;
+				Type = (EP4ParamType::Type) ((int)Type + 1);
+				ParamDetectionIteratorsStack.Add(IP4EnvParamDetectionIterator::Create(Type, CommandLine, *this));
+				continue;
 			}
-
-			switch (Type)
+			else
 			{
-			case EP4ParamType::Path:
-				Stop = !DetectPath(FieldReference, Env);
-				break;
-			case EP4ParamType::Port:
-				Stop = !DetectPort(FieldReference, Env);
-				break;
-			case EP4ParamType::User:
-				Stop = !DetectUser(FieldReference, Env);
-				break;
-			case EP4ParamType::Client:
-				Stop = !DetectClient(FieldReference, Env);
-				break;
-			case EP4ParamType::Branch:
-				Stop = !DetectBranch(FieldReference, Env);
-				break;
+				return true;
 			}
 		}
-
-		FAutoDetectTasks(FP4Env& Env)
-			: Env(Env)
+		else
 		{
-			Stop = false;
+			Type = (EP4ParamType::Type) ((int)Type - 1);
+			ParamDetectionIteratorsStack.RemoveAt(ParamDetectionIteratorsStack.Num() - 1);
 		}
+	}
 
-		FP4Env& Env;
-		bool Stop;
-	};
-
-	FAutoDetectTasks AT(*this);
-	SerializeParams(FSerializationTask::CreateRaw(&AT, &FAutoDetectTasks::AutoDetectParam));
-
-	return !AT.Stop;
+	return false;
 }
 
 FString FP4Env::GetCommandLine()
@@ -507,5 +792,51 @@ const FString& FP4Env::GetPath() const
 	return Path;
 }
 
+void FP4Env::SetParam(EP4ParamType::Type Type, const FString& Value)
+{
+	switch (Type)
+	{
+	case EP4ParamType::Path:
+		Path = Value;
+		break;
+	case EP4ParamType::Port:
+		Port = Value;
+		break;
+	case EP4ParamType::User:
+		User = Value;
+		break;
+	case EP4ParamType::Client:
+		Client = Value;
+		break;
+	case EP4ParamType::Branch:
+		Branch = Value;
+		break;
+	default:
+		// Unimplemented param type?
+		checkNoEntry();
+	}
+}
+
+TSharedPtr<IP4EnvParamDetectionIterator> IP4EnvParamDetectionIterator::Create(EP4ParamType::Type Type, const TCHAR* CommandLine, const FP4Env& Env)
+{
+	switch (Type)
+	{
+	case EP4ParamType::Path:
+		return MakeShareable(new FP4PathDetectionIterator(CommandLine));
+	case EP4ParamType::Port:
+		return MakeShareable(new FP4PortDetectionIterator(CommandLine));
+	case EP4ParamType::User:
+		return MakeShareable(new FP4UserDetectionIterator(CommandLine, Env));
+	case EP4ParamType::Client:
+		return MakeShareable(new FP4ClientDetectionIterator(CommandLine, Env));
+	case EP4ParamType::Branch:
+		return MakeShareable(new FP4BranchDetectionIterator(CommandLine, Env));
+	default:
+		// Unimplemented param type?
+		checkNoEntry();
+		return nullptr;
+	}
+}
+
 /* Static variable initialization. */
 TSharedPtr<FP4Env> FP4Env::Env;
\ No newline at end of file
diff --git a/Engine/Source/Programs/UnrealSync/Private/P4Env.h b/Engine/Source/Programs/UnrealSync/Private/P4Env.h
--- a/Engine/Source/Programs/UnrealSync/Private/P4Env.h
+++ b/Engine/Source/Programs/UnrealSync/Private/P4Env.h
@@ -12,11 +12,11 @@ namespace EP4ParamType
 {
 	enum Type
 	{
-		Path,
-		Port,
-		User,
-		Client,
-		Branch
+		Path	= 0,
+		Port	= 1,
+		User	= 2,
+		Client	= 3,
+		Branch	= 4
 	};
 }
 
@@ -130,16 +130,23 @@ public:
 	 */
 	FString GetCommandLine();
 
+	/**
+	 * Gets param type name.
+	 *
+	 * @param Type Param type enum id.
+	 *
+	 * @returns Param type name.
+	 */
+	static FString GetParamName(EP4ParamType::Type Type);
+
 private:
 	/* Param serialization delegate. */
 	DECLARE_DELEGATE_TwoParams(FSerializationTask, FString&, EP4ParamType::Type);
 
 	/**
 	 * Constructor
-	 *
-	 * @param CommandLine Command line to parse params from.
 	 */
-	FP4Env(const TCHAR* CommandLine);
+	FP4Env();
 
 	/**
 	 * Serializes all P4Env params.
@@ -151,29 +158,19 @@ private:
 	/**
 	 * Auto-detect missing params.
 	 *
-	 * @returns True if succeeded. False otherwise.
-	 */
-	bool AutoDetectMissingParams();
-
-	/**
-	 * Gets param type name.
-	 *
-	 * @param Type Param type enum id.
+	 * @param CommandLine Command line that the program was run with.
 	 *
-	 * @returns Param type name.
+	 * @returns True if succeeded. False otherwise.
 	 */
-	static FString GetParamName(EP4ParamType::Type Type);
+	bool AutoDetectMissingParams(const TCHAR* CommandLine);
 
 	/**
-	 * Parses param either from command line or environment variables.
-	 * Command line has priority.
-	 *
-	 * @param CommandLine Command line to parse.
-	 * @param ParamName Param name to parse.
+	 * Sets param value.
 	 *
-	 * @returns Parsed value.
+	 * @param Type Type of param to set.
+	 * @param Value Value to set.
 	 */
-	static FString GetParam(const TCHAR* CommandLine, const FString& ParamName);
+	void SetParam(EP4ParamType::Type Type, const FString& Value);
 
 	/* Singleton instance pointer. */
 	static TSharedPtr<FP4Env> Env;
