UnrealSync fix: Some variables weren't cleared at start-up.

diff --git a/Engine/Source/Programs/UnrealSync/Private/P4Env.cpp b/Engine/Source/Programs/UnrealSync/Private/P4Env.cpp
--- a/Engine/Source/Programs/UnrealSync/Private/P4Env.cpp
+++ b/Engine/Source/Programs/UnrealSync/Private/P4Env.cpp
@@ -237,7 +237,7 @@ public:
 	 * @param CommandLine Command line to parse.
 	 */
 	FP4PathDetectionIterator(const TCHAR* CommandLine)
-		: FP4EnvParamDetectionIteratorBase(EP4ParamType::Path, CommandLine)
+		: Step(0), FP4EnvParamDetectionIteratorBase(EP4ParamType::Path, CommandLine)
 	{
 		PossibleLocations.Add("C:\\Program Files\\Perforce");
 		PossibleLocations.Add("C:\\Program Files (x86)\\Perforce");
@@ -398,7 +398,7 @@ public:
 	 * @param CommandLine Command line to parse.
 	 */
 	FP4PortDetectionIterator(const TCHAR* CommandLine)
-		: FP4EnvParamDetectionIteratorBase(EP4ParamType::Port, CommandLine)
+		: Step(0), FP4EnvParamDetectionIteratorBase(EP4ParamType::Port, CommandLine)
 	{ }
 
 	/**
@@ -450,7 +450,7 @@ public:
 	 * @param Env Current P4 environment state.
 	 */
 	FP4UserDetectionIterator(const TCHAR* CommandLine, const FP4Env& Env)
-		: FP4EnvParamDetectionIteratorBase(EP4ParamType::User, CommandLine), Env(Env)
+		: Step(0), FP4EnvParamDetectionIteratorBase(EP4ParamType::User, CommandLine), Env(Env)
 	{
 	}
 
