UnrealSync changes as requested here: https://udn.unrealengine.com/questions/223609/changes-to-unrealsync-and-objectversioncpp.html.

diff --git a/Engine/Source/Programs/UnrealSync/Private/P4Env.cpp b/Engine/Source/Programs/UnrealSync/Private/P4Env.cpp
--- a/Engine/Source/Programs/UnrealSync/Private/P4Env.cpp
+++ b/Engine/Source/Programs/UnrealSync/Private/P4Env.cpp
@@ -548,20 +548,20 @@ public:
 			if (KnownPath.StartsWith(FPaths::ConvertRelativePathToFull(Root)))
 			{
 				FString InfoOutput;
-				if (!RunProcessOutput(Env.GetPath(), FString::Printf(TEXT("-p%s -u%s -c%s info"), *Env.GetPort(), *Env.GetUser(), *ClientName), InfoOutput))
+				if (!RunProcessOutput(Env.GetPath(), FString::Printf(TEXT("-p%s -u%s client -o %s"), *Env.GetPort(), *Env.GetUser(), *ClientName), InfoOutput))
 				{
 					continue;
 				}
 
-				const FRegexPattern InfoPattern(TEXT("Client host: ([^\\r\\n ]+)\\s"));
+				const FRegexPattern InfoPattern(TEXT("Host:\\s*([^\\r\\n\\t ]+)\\s*"));
 				FRegexMatcher InfoMatcher(InfoPattern, InfoOutput);
 
-				if (InfoMatcher.FindNext())
+				while (InfoMatcher.FindNext())
 				{
 					if (InfoOutput.Mid(
 						InfoMatcher.GetCaptureGroupBeginning(1),
 						InfoMatcher.GetCaptureGroupEnding(1) - InfoMatcher.GetCaptureGroupBeginning(1)
-						).Equals(HostName))
+						).Equals(HostName, ESearchCase::IgnoreCase))
 					{
 						Output = ClientName;
 						return true;
@@ -740,6 +740,8 @@ bool FP4Env::AutoDetectMissingParams(const TCHAR* CommandLine)
 	EP4ParamType::Type Type = EP4ParamType::Path;
 	ParamDetectionIteratorsStack.Add(IP4EnvParamDetectionIterator::Create(Type, CommandLine, *this));
 
+	int32 IterationCountdown = 20;
+
 	while (ParamDetectionIteratorsStack.Num() > 0)
 	{
 		if (ParamDetectionIteratorsStack.Last()->MoveNext())
@@ -761,6 +763,11 @@ bool FP4Env::AutoDetectMissingParams(const TCHAR* CommandLine)
 		{
 			Type = (EP4ParamType::Type) ((int)Type - 1);
 			ParamDetectionIteratorsStack.RemoveAt(ParamDetectionIteratorsStack.Num() - 1);
+			--IterationCountdown;
+			if (!IterationCountdown)
+			{
+				break;
+			}
 		}
 	}
 
