UnrealSync fix. Compile against CoreUObject.

UE-6171

Something had to changed in the dependency chain and now we have to set compile against CoreUObject explicitly.

diff --git a/Engine/Source/Programs/UnrealSync/UnrealSync.Target.cs b/Engine/Source/Programs/UnrealSync/UnrealSync.Target.cs
--- a/Engine/Source/Programs/UnrealSync/UnrealSync.Target.cs
+++ b/Engine/Source/Programs/UnrealSync/UnrealSync.Target.cs
@@ -51,7 +51,7 @@ public class UnrealSyncTarget : TargetRules
 
 		// Currently this app is not linking against the engine, so we'll compile out references from Core to the rest of the engine
 		UEBuildConfiguration.bCompileAgainstEngine = false;
-		UEBuildConfiguration.bCompileAgainstCoreUObject = false;
+		UEBuildConfiguration.bCompileAgainstCoreUObject = true;
 
 		OutLinkEnvironmentConfiguration.bIsBuildingConsoleApplication = false;
 	}
