fixed typo

diff --git a/Engine/Source/Programs/UnrealSync/Private/P4Env.cpp b/Engine/Source/Programs/UnrealSync/Private/P4Env.cpp
--- a/Engine/Source/Programs/UnrealSync/Private/P4Env.cpp
+++ b/Engine/Source/Programs/UnrealSync/Private/P4Env.cpp
@@ -589,7 +589,7 @@ private:
  * Gets branch detected for this app from current P4 environment state.
  *
  * @param Output Found branch prefix.
- * @parma Env Current P4 environment state.
+ * @param Env Current P4 environment state.
  *
  * @returns True if found, false otherwise.
  */
diff --git a/Engine/Source/Runtime/Core/Public/Math/BigInt.h b/Engine/Source/Runtime/Core/Public/Math/BigInt.h
--- a/Engine/Source/Runtime/Core/Public/Math/BigInt.h
+++ b/Engine/Source/Runtime/Core/Public/Math/BigInt.h
@@ -91,7 +91,7 @@ public:
 	/**
 	 * Constructor. Initializes this big int with a 64 bit integer value.
 	 *
-	 * @parma Other The value to set.
+	 * @param Other The value to set.
 	 */
 	TBigInt(int64 Other)
 	{
