UnrealSync: Changes in Slate made it impossible for non-Slate threads to request list refresh. This was used in US and caused crashes. Fixed.

diff --git a/Engine/Source/Programs/UnrealSync/Private/GUI.cpp b/Engine/Source/Programs/UnrealSync/Private/GUI.cpp
--- a/Engine/Source/Programs/UnrealSync/Private/GUI.cpp
+++ b/Engine/Source/Programs/UnrealSync/Private/GUI.cpp
@@ -816,6 +816,81 @@ bool DeleteIfExistsAndCopyFile(const FString& To, const FString& From)
 	return PlatformPhysical.CopyFile(*To, *From);
 }
 
+/**
+ * Widget displaying log lines.
+ */
+template <typename TElement>
+class SLogList : public SListView<TElement>
+{
+public:
+	SLATE_BEGIN_ARGS(SLogList){}
+		SLATE_ARGUMENT(const TArray<TElement>*, ListItemsSource)
+		SLATE_EVENT(FOnGenerateRow, OnGenerateRow)
+	SLATE_END_ARGS()
+
+	BEGIN_SLATE_FUNCTION_BUILD_OPTIMIZATION
+	void Construct(const FArguments& InArgs)
+	{
+		SListView<TElement>::Construct(
+				SListView<TElement>::FArguments()
+					.ListItemsSource(InArgs._ListItemsSource)
+					.OnGenerateRow(InArgs._OnGenerateRow)
+			);
+
+		RegisterActiveTimer(0.5f, FWidgetActiveTimerDelegate::CreateSP(this, &SLogList::ExecuteRequests));
+	}
+	END_SLATE_FUNCTION_BUILD_OPTIMIZATION
+
+	/** Delegate that represents rendering thread request. */
+	DECLARE_DELEGATE(FRenderingThreadRequest);
+
+	/**
+	 * Queues request to ScrollIntoView for inner list to be executed by Slate
+	 * thread.
+	 */
+	void ExternalThreadRequestScrollIntoView(TElement Element)
+	{
+		AddRenderingThreadRequest(FRenderingThreadRequest::CreateSP(this, &SListView<TElement>::RequestScrollIntoView, Element));
+	}
+
+private:
+	/**
+	 * Executes all pending requests.
+	 *
+	 * Some Slate methods are restricted for either game or Slate rendering
+	 * thread. If external thread tries to execute them assert raises, which is
+	 * telling that it's unsafe.
+	 *
+	 * This method executes all pending requests on Slate thread, so you can
+	 * use those to execute some restricted methods.
+	 */
+	EActiveTimerReturnType ExecuteRequests(double CurrentTime, float DeltaTime)
+	{
+		FScopeLock Lock(&RequestsMutex);
+		while (Requests.Num())
+		{
+			Requests.Pop().Execute();
+		}
+
+		return EActiveTimerReturnType::Continue;
+	}
+
+	/**
+	 * Adds delegate to pending requests list.
+	 */
+	void AddRenderingThreadRequest(FRenderingThreadRequest Request)
+	{
+		FScopeLock Lock(&RequestsMutex);
+		Requests.Add(MoveTemp(Request));
+	}
+
+	/** Mutex for pending requests list. */
+	FCriticalSection RequestsMutex;
+
+	/** Pending requests list. */
+	TArray<FRenderingThreadRequest> Requests;
+};
+
 /**
  * Main tab widget.
  */
@@ -874,7 +949,7 @@ public:
 				SNew(STextBlock).Text(FText::FromString("Go back"))
 			];
 
-		LogListView = SNew(SListView<TSharedPtr<FString> >)
+		LogListView = SNew(SLogList<TSharedPtr<FString> >)
 			.ListItemsSource(&LogLines)
 			.OnGenerateRow(this, &SMainTabWidget::GenerateLogItem);
 		
@@ -1170,8 +1245,7 @@ private:
 			Buffer = Rest;
 		}
 
-		LogListView->RequestListRefresh();
-		LogListView->RequestScrollIntoView(LogLines.Last());
+		LogListView->ExternalThreadRequestScrollIntoView(LogLines.Last());
 
 		return true;
 	}
@@ -1267,7 +1341,7 @@ private:
 	/* Report log. */
 	TArray<TSharedPtr<FString> > LogLines;
 	/* Log list view. */
-	TSharedPtr<SListView<TSharedPtr<FString> > > LogListView;
+	TSharedPtr<SLogList<TSharedPtr<FString> > > LogListView;
 	/* Go back button reference. */
 	TSharedPtr<SButton> GoBackButton;
 
