UE-1539: Selecting Sync takes 3-5 mins even when nothing is pulled down.

Merge a couple of p4.exe sync calls into one, which should save some time.

diff --git a/Engine/Source/Programs/UnrealSync/Private/UnrealSync.cpp b/Engine/Source/Programs/UnrealSync/Private/UnrealSync.cpp
--- a/Engine/Source/Programs/UnrealSync/Private/UnrealSync.cpp
+++ b/Engine/Source/Programs/UnrealSync/Private/UnrealSync.cpp
@@ -428,12 +428,14 @@ bool FUnrealSync::Sync(bool bArtist, bool bPreview, const FString& Label, const
 		SyncSteps.Add("/..." + ProgramRevisionSpec); // all files to label
 	}
 
-	for(auto SyncStep : SyncSteps)
+	for(auto& SyncStep : SyncSteps)
 	{
-		if (!FP4Env::RunP4Progress(FString("sync ") + (bPreview ? "-n " : "") + FP4Env::Get().GetBranch() + SyncStep, OnSyncProgress))
-		{
-			return false;
-		}
+		SyncStep = FP4Env::Get().GetBranch() + SyncStep;
+	}
+
+	if (!FP4Env::RunP4Progress(FString("sync ") + (bPreview ? "-n " : "") + FString::Join(SyncSteps, TEXT(" ")), OnSyncProgress))
+	{
+		return false;
 	}
 
 	return true;
