UECORE-125: UnrealSync: Change check boxes that work like radio buttons to actually look like radio buttons

+ A couple of cosmetic changes.

diff --git a/Engine/Source/Programs/UnrealSync/Private/GUI.cpp b/Engine/Source/Programs/UnrealSync/Private/GUI.cpp
--- a/Engine/Source/Programs/UnrealSync/Private/GUI.cpp
+++ b/Engine/Source/Programs/UnrealSync/Private/GUI.cpp
@@ -172,6 +172,7 @@ public:
 				+ SVerticalBox::Slot().AutoHeight()
 				[
 					SNew(SCheckBox)
+					.Style(FCoreStyle::Get(), "RadioButton")
 					.OnCheckStateChanged(this, &FItem::OnCheckStateChange)
 					.IsChecked(this, &FItem::IsCheckboxChecked)
 					[
@@ -179,7 +180,7 @@ public:
 						.Text(FText::FromString(GetName()))
 					]
 				]
-			+ SVerticalBox::Slot().VAlign(VAlign_Fill)
+				+ SVerticalBox::Slot().VAlign(VAlign_Fill).Padding(19.0f, 2.0f, 2.0f, 2.0f)
 				[
 					GetContent()
 				];
@@ -957,7 +958,7 @@ public:
 			+ SWidgetSwitcher::Slot()
 			[
 				SNew(SVerticalBox)
-				+ SVerticalBox::Slot().AutoHeight()
+				+ SVerticalBox::Slot().AutoHeight().Padding(5.0f)
 				[
 					PickGameWidget.ToSharedRef()
 				]
