UnrealSync problem with using non-threadsafe shared pointers in threading. Fixed.

diff --git a/Engine/Source/Programs/UnrealSync/Private/GUI.cpp b/Engine/Source/Programs/UnrealSync/Private/GUI.cpp
--- a/Engine/Source/Programs/UnrealSync/Private/GUI.cpp
+++ b/Engine/Source/Programs/UnrealSync/Private/GUI.cpp
@@ -818,87 +818,63 @@ bool DeleteIfExistsAndCopyFile(const FString& To, const FString& From)
 }
 
 /**
- * Widget displaying log lines.
+ * Main tab widget.
  */
-template <typename TElement>
-class SLogList : public SListView<TElement>
+class SMainTabWidget : public SCompoundWidget
 {
-public:
-	SLATE_BEGIN_ARGS(SLogList){}
-		SLATE_ARGUMENT(const TArray<TElement>*, ListItemsSource)
-		SLATE_EVENT(FOnGenerateRow, OnGenerateRow)
-	SLATE_END_ARGS()
-
-	BEGIN_SLATE_FUNCTION_BUILD_OPTIMIZATION
-	void Construct(const FArguments& InArgs)
-	{
-		SListView<TElement>::Construct(
-				SListView<TElement>::FArguments()
-					.ListItemsSource(InArgs._ListItemsSource)
-					.OnGenerateRow(InArgs._OnGenerateRow)
-			);
-
-		RegisterActiveTimer(0.5f, FWidgetActiveTimerDelegate::CreateSP(this, &SLogList::ExecuteRequests));
-	}
-	END_SLATE_FUNCTION_BUILD_OPTIMIZATION
-
-	/** Delegate that represents rendering thread request. */
-	DECLARE_DELEGATE(FRenderingThreadRequest);
-
 	/**
-	 * Queues request to ScrollIntoView for inner list to be executed by Slate
-	 * thread.
+	 * External threads requests dispatcher.
 	 */
-	void ExternalThreadRequestScrollIntoView(TElement Element)
+	class FExternalThreadsDispatcher : public TSharedFromThis<FExternalThreadsDispatcher, ESPMode::ThreadSafe>
 	{
-		AddRenderingThreadRequest(FRenderingThreadRequest::CreateSP(this, &SListView<TElement>::RequestScrollIntoView, Element));
-	}
+	public:
+		/** Delegate that represents rendering thread request. */
+		DECLARE_DELEGATE(FRenderingThreadRequest);
 
-private:
-	/**
-	 * Executes all pending requests.
-	 *
-	 * Some Slate methods are restricted for either game or Slate rendering
-	 * thread. If external thread tries to execute them assert raises, which is
-	 * telling that it's unsafe.
-	 *
-	 * This method executes all pending requests on Slate thread, so you can
-	 * use those to execute some restricted methods.
-	 */
-	EActiveTimerReturnType ExecuteRequests(double CurrentTime, float DeltaTime)
-	{
-		FScopeLock Lock(&RequestsMutex);
-		while (Requests.Num())
+		/**
+		 * Executes all pending requests.
+		 *
+		 * Some Slate methods are restricted for either game or Slate rendering
+		 * thread. If external thread tries to execute them assert raises, which is
+		 * telling that it's unsafe.
+		 *
+		 * This method executes all pending requests on Slate thread, so you can
+		 * use those to execute some restricted methods.
+		 */
+		EActiveTimerReturnType ExecuteRequests(double CurrentTime, float DeltaTime)
 		{
-			Requests.Pop().Execute();
-		}
+			FScopeLock Lock(&RequestsMutex);
+			while (Requests.Num())
+			{
+				Requests.Pop().Execute();
+			}
 
-		return EActiveTimerReturnType::Continue;
-	}
+			return EActiveTimerReturnType::Continue;
+		}
 
-	/**
-	 * Adds delegate to pending requests list.
-	 */
-	void AddRenderingThreadRequest(FRenderingThreadRequest Request)
-	{
-		FScopeLock Lock(&RequestsMutex);
-		Requests.Add(MoveTemp(Request));
-	}
+		/**
+		 * Adds delegate to pending requests list.
+		 */
+		void AddRenderingThreadRequest(FRenderingThreadRequest Request)
+		{
+			FScopeLock Lock(&RequestsMutex);
+			Requests.Add(MoveTemp(Request));
+		}
 
-	/** Mutex for pending requests list. */
-	FCriticalSection RequestsMutex;
+	private:
+		/** Mutex for pending requests list. */
+		FCriticalSection RequestsMutex;
 
-	/** Pending requests list. */
-	TArray<FRenderingThreadRequest> Requests;
-};
+		/** Pending requests list. */
+		TArray<FRenderingThreadRequest> Requests;
+	};
 
-/**
- * Main tab widget.
- */
-class SMainTabWidget : public SCompoundWidget
-{
 public:
-	SLATE_BEGIN_ARGS(SMainTabWidget){}
+	SMainTabWidget()
+		: ExternalThreadsDispatcher(MakeShareable(new FExternalThreadsDispatcher()))
+	{}
+
+	SLATE_BEGIN_ARGS(SMainTabWidget) {}
 	SLATE_END_ARGS()
 
 	BEGIN_SLATE_FUNCTION_BUILD_OPTIMIZATION
@@ -950,7 +926,7 @@ public:
 				SNew(STextBlock).Text(FText::FromString("Go back"))
 			];
 
-		LogListView = SNew(SLogList<TSharedPtr<FString> >)
+		LogListView = SNew(SListView<TSharedPtr<FString> >)
 			.ListItemsSource(&LogLines)
 			.OnGenerateRow(this, &SMainTabWidget::GenerateLogItem);
 		
@@ -1030,6 +1006,8 @@ public:
 		FUnrealSync::RegisterOnDataReset(FUnrealSync::FOnDataReset::CreateRaw(this, &SMainTabWidget::DataReset));
 		FUnrealSync::RegisterOnDataLoaded(FUnrealSync::FOnDataLoaded::CreateRaw(this, &SMainTabWidget::DataLoaded));
 
+		RegisterActiveTimer(0.5f, FWidgetActiveTimerDelegate::CreateThreadSafeSP(&ExternalThreadsDispatcher.Get(), &FExternalThreadsDispatcher::ExecuteRequests));
+
 		OnReloadLabels();
 	}
 	END_SLATE_FUNCTION_BUILD_OPTIMIZATION
@@ -1043,6 +1021,29 @@ public:
 	}
 
 private:
+	/**
+	 * Queues adding lines to the log for execution on Slate rendering thread.
+	 */
+	void ExternalThreadAddLinesToLog(TArray<TSharedPtr<FString> > Lines)
+	{
+		ExternalThreadsDispatcher->AddRenderingThreadRequest(FExternalThreadsDispatcher::FRenderingThreadRequest::CreateSP(this, &SMainTabWidget::AddLinesToLog, Lines));
+	}
+
+	/**
+	 * Adds lines to the log.
+	 *
+	 * @param Lines Lines to add.
+	 */
+	void AddLinesToLog(TArray<TSharedPtr<FString> > Lines)
+	{
+		for (const auto& Line : Lines)
+		{
+			LogLines.Add(Line);
+		}
+
+		LogListView->RequestScrollIntoView(LogLines.Last());
+	}
+
 	/**
 	 * This method does some initial task:
 	 * - initializes P4
@@ -1239,14 +1240,16 @@ private:
 		FString Line;
 		FString Rest;
 
+		TArray<TSharedPtr<FString> > Lines;
+
 		while (Buffer.Split("\n", &Line, &Rest))
 		{
-			LogLines.Add(MakeShareable(new FString(Line)));
+			Lines.Add(MakeShareable(new FString(Line)));
 
 			Buffer = Rest;
 		}
 
-		LogListView->ExternalThreadRequestScrollIntoView(LogLines.Last());
+		ExternalThreadAddLinesToLog(Lines);
 
 		return true;
 	}
@@ -1339,10 +1342,14 @@ private:
 	/* Check box to tell if this should be a preview sync. */
 	TSharedPtr<SCheckBox> PreviewSyncCheckBox;
 
+	/* External thread requests dispatcher. */
+	TSharedRef<FExternalThreadsDispatcher, ESPMode::ThreadSafe> ExternalThreadsDispatcher;
+
 	/* Report log. */
 	TArray<TSharedPtr<FString> > LogLines;
 	/* Log list view. */
-	TSharedPtr<SLogList<TSharedPtr<FString> > > LogListView;
+	TSharedPtr<SListView<TSharedPtr<FString> > > LogListView;
+
 	/* Go back button reference. */
 	TSharedPtr<SButton> GoBackButton;
 
