UECORE-122: UnrealSync P4 settings

diff --git a/Engine/Source/Programs/UnrealSync/Private/GUI.cpp b/Engine/Source/Programs/UnrealSync/Private/GUI.cpp
--- a/Engine/Source/Programs/UnrealSync/Private/GUI.cpp
+++ b/Engine/Source/Programs/UnrealSync/Private/GUI.cpp
@@ -711,112 +711,6 @@ private:
 	FOnGamePicked OnGamePicked;
 };
 
-/**
- * Creates Slate page to show error and close app.
- *
- * @param ErrorMsg Error message to display.
- *
- * @returns Created widget.
- */
-TSharedRef<SWidget> CreateErrorPage(const FString& ErrorMsg)
-{
-	/**
-	 * Embedded class to enclose closing button functionality.
-	 */
-	class CloseButtonHandler
-	{
-	public:
-		/**
-		 * Function that is called on when Close button is clicked.
-		 * It tells the app to close.
-		 *
-		 * @returns Tells that this event was handled.
-		 */
-		static FReply OnCloseButtonClick()
-		{
-			FPlatformMisc::RequestExit(false);
-
-			return FReply::Handled();
-		}
-	};
-
-	return
-		SNew(SVerticalBox)
-		+ SVerticalBox::Slot().HAlign(HAlign_Center).VAlign(VAlign_Center)
-		[
-			SNew(STextBlock).Text(FText::FromString(*ErrorMsg))
-		]
-		+ SVerticalBox::Slot().AutoHeight()
-		[
-			SNew(SHorizontalBox)
-			+ SHorizontalBox::Slot()
-			+ SHorizontalBox::Slot().AutoWidth()
-			[
-				SNew(SButton)
-				.OnClicked_Static(&CloseButtonHandler::OnCloseButtonClick)
-				[
-					SNew(STextBlock).Text(FText::FromString("Close"))
-				]
-			]
-		];
-}
-
-/**
- * Tries to update original UnrealSync at given location.
- *
- * @param Location of original UnrealSync.
- *
- * @returns False on failure. True otherwise.
- */
-bool UpdateOriginalUS(const FString& OriginalUSPath)
-{
-	FString Output;
-	return FP4Env::RunP4Output(TEXT("sync ") + OriginalUSPath, Output);
-}
-
-/**
- * Runs detached UnrealSync process and passes given parameters in command line.
- *
- * @param USPath UnrealSync executable path.
- * @param bDoNotRunFromCopy Should UnrealSync be called with -DoNotRunFromCopy flag?
- * @param bDoNotUpdateOnStartUp Should UnrealSync be called with -DoNotUpdateOnStartUp flag?
- * @param bPassP4Env Pass P4 environment parameters to UnrealSync?
- *
- * @returns True if succeeded. False otherwise. Notice that this says of
- *			success of launching procedure not launched process, cause
- *			its detached.
- */
-bool RunDetachedUS(const FString& USPath, bool bDoNotRunFromCopy, bool bDoNotUpdateOnStartUp, bool bPassP4Env)
-{
-	FString CommandLine = FString()
-		+ (bDoNotRunFromCopy ? TEXT("-DoNotRunFromCopy ") : TEXT(""))
-		+ (bDoNotUpdateOnStartUp ? TEXT("-DoNotUpdateOnStartUp ") : TEXT(""))
-		+ (bPassP4Env ? *FP4Env::Get().GetCommandLine()	: TEXT(""));
-
-	return RunProcess(USPath, CommandLine);
-}
-
-/**
- * This function copies file from From to To location and deletes
- * if To exists.
- *
- * @param To Location to which copy the file.
- * @param From Location from which copy the file.
- *
- * @returns True if succeeded. False otherwise.
- */
-bool DeleteIfExistsAndCopyFile(const FString& To, const FString& From)
-{
-	auto& PlatformPhysical = IPlatformFile::GetPlatformPhysical();
-
-	if (PlatformPhysical.FileExists(*To) && !PlatformPhysical.DeleteFile(*To))
-	{
-		return false;
-	}
-
-	return PlatformPhysical.CopyFile(*To, *From);
-}
-
 /**
  * Main tab widget.
  */
@@ -880,22 +774,6 @@ public:
 	BEGIN_SLATE_FUNCTION_BUILD_OPTIMIZATION
 	void Construct(const FArguments& InArgs)
 	{
-		FString ErrorMsg;
-		if (DoInitialTasks(ErrorMsg))
-		{
-			FPlatformMisc::RequestExit(false);
-			return;
-		}
-
-		if (!ErrorMsg.IsEmpty())
-		{
-			this->ChildSlot
-				[
-					CreateErrorPage(ErrorMsg)
-				];
-			return;
-		}
-
 		RadioSelection = SNew(SRadioContentSelection);
 
 		AddToRadioSelection("Sync to the latest promoted", SNew(SLatestPromoted));
@@ -1044,75 +922,6 @@ private:
 		LogListView->RequestScrollIntoView(LogLines.Last());
 	}
 
-	/**
-	 * This method does some initial task:
-	 * - initializes P4
-	 * - updates UnrealSync (to do this it has to copy itself to remote location and run from there)
-	 *
-	 * @param ErrorMsg Out parameter to output error message if happened.
-	 *
-	 * @returns True if application should stop building GUI and exit. False otherwise.
-	 */
-	bool DoInitialTasks(FString& ErrorMsg)
-	{
-		bool bDoNotUpdateOnStartUp = FParse::Param(FCommandLine::Get(), TEXT("DoNotUpdateOnStartUp"));
-		bool bDoNotRunFromCopy = FParse::Param(FCommandLine::Get(), TEXT("DoNotRunFromCopy"));
-
-		if (!FP4Env::Init(FCommandLine::Get()))
-		{
-			ErrorMsg = TEXT("P4 environment detection failed.");
-			return false;
-		}
-
-		FString CommonExecutablePath =
-			FPaths::ConvertRelativePathToFull(
-			FPaths::Combine(*FPaths::EngineDir(), TEXT("Binaries"), TEXT("Win64"),
-#if !UE_BUILD_DEBUG
-				TEXT("UnrealSync")
-#else
-				TEXT("UnrealSync-Win64-Debug")
-#endif
-			));
-
-		FString OriginalExecutablePath = CommonExecutablePath + TEXT(".exe");
-		FString TemporaryExecutablePath = CommonExecutablePath + TEXT(".Temporary.exe");
-
-		if (!bDoNotRunFromCopy)
-		{
-			if (!DeleteIfExistsAndCopyFile(*TemporaryExecutablePath, *OriginalExecutablePath))
-			{
-				ErrorMsg = TEXT("Copying UnrealSync to temp location failed.");
-				return false;
-			}
-
-			if (!RunDetachedUS(TemporaryExecutablePath, true, bDoNotUpdateOnStartUp, true))
-			{
-				ErrorMsg = TEXT("Running remote UnrealSync failed.");
-				return false;
-			}
-
-			return true;
-		}
-
-		if (!bDoNotUpdateOnStartUp)
-		{
-			if (!UpdateOriginalUS(OriginalExecutablePath))
-			{
-				ErrorMsg = TEXT("UnrealSync update failed.");
-				return false;
-			}
-			if (!RunDetachedUS(OriginalExecutablePath, false, true, true))
-			{
-				ErrorMsg = TEXT("Running UnrealSync failed.");
-				return false;
-			}
-
-			return true;
-		}
-
-		return false;
-	}
-
 	/**
 	 * Generates list view item from FString.
 	 *
@@ -1369,46 +1178,69 @@ TSharedRef<SDockTab> GetMainTab(const FSpawnTabArgs& Args)
 	TSharedRef<SDockTab> MainTab =
 		SNew(SDockTab)
 		.TabRole(ETabRole::MajorTab)
-		.Label(FText::FromString("UnrealSync"))
+		.Label(FText::FromString("Syncing"))
 		.ToolTipText(FText::FromString("Sync Unreal Engine tool."));
 
 	MainTab->SetContent(
 		SNew(SMainTabWidget)
 		);
 
+	FGlobalTabmanager::Get()->SetActiveTab(MainTab);
+
 	return MainTab;
 }
 
 /**
- * Initializes and starts GUI.
+ * Creates and returns P4 settings tab.
+ *
+ * @param Args Args used to spawn tab.
  *
- * @param CommandLine Command line passed to the program.
+ * @returns P4 settings UnrealSync tab.
  */
-void InitGUI(const TCHAR* CommandLine)
+TSharedRef<SDockTab> GetP4EnvTab(const FSpawnTabArgs& Args)
 {
-	// start up the main loop
-	GEngineLoop.PreInit(CommandLine);
+	TSharedRef<SDockTab> P4EnvTab =
+		SNew(SDockTab)
+		.TabRole(ETabRole::MajorTab)
+		.Label(FText::FromString("Perforce Settings"))
+		.ToolTipText(FText::FromString("Perforce settings."));
+
+	P4EnvTab->SetContent(
+		SNew(SP4EnvTabWidget)
+		);
+
+	return P4EnvTab;
+}
 
+void InitGUI(const TCHAR* CommandLine, bool bP4EnvTabOnly)
+{
 	// crank up a normal Slate application using the platform's standalone renderer
 	FSlateApplication::InitializeAsStandaloneApplication(GetStandardStandaloneRenderer());
 
 	// set the application name
 	FGlobalTabmanager::Get()->SetApplicationTitle(NSLOCTEXT("UnrealSync", "AppTitle", "Unreal Sync"));
 
-	FGlobalTabmanager::Get()->RegisterTabSpawner("MainTab", FOnSpawnTab::CreateStatic(&GetMainTab));
+	auto TabStack = FTabManager::NewStack();
+
+	if (!bP4EnvTabOnly)
+	{
+		TabStack->AddTab("MainTab", ETabState::OpenedTab);
+		FGlobalTabmanager::Get()->RegisterTabSpawner("MainTab", FOnSpawnTab::CreateStatic(&GetMainTab));
+		TabStack->SetForegroundTab(FTabId("MainTab"));
+	}
+
+	FGlobalTabmanager::Get()->RegisterTabSpawner("P4EnvTab", FOnSpawnTab::CreateStatic(&GetP4EnvTab));
+	TabStack->AddTab("P4EnvTab", ETabState::OpenedTab);
 
 	TSharedRef<FTabManager::FLayout> Layout = FTabManager::NewLayout("UnrealSyncLayout")
 		->AddArea(
 			FTabManager::NewArea(720, 370)
 			->SetWindow(FVector2D(420, 10), false)
-			->Split(
-				FTabManager::NewStack()
-				->AddTab("MainTab", ETabState::OpenedTab)
-			)
+			->Split(TabStack)
 		);
 
 	FGlobalTabmanager::Get()->RestoreFrom(Layout, TSharedPtr<SWindow>());
-
+	
 	// loop while the server does the rest
 	while (!GIsRequestingExit)
 	{
diff --git a/Engine/Source/Programs/UnrealSync/Private/GUI.h b/Engine/Source/Programs/UnrealSync/Private/GUI.h
--- a/Engine/Source/Programs/UnrealSync/Private/GUI.h
+++ b/Engine/Source/Programs/UnrealSync/Private/GUI.h
@@ -3,8 +3,9 @@
 #pragma once
 
 /**
- * Function that is initializing and constructing the whole GUI.
+ * Initializes and starts GUI.
  *
- * @param CommandLine Command line that has been passed to the process.
+ * @param CommandLine Command line passed to the program.
+ * @param bP4EnvTabOnly Tells if initialize GUI with P4 env settings tab only.
  */
-void InitGUI(const TCHAR* CommandLine);
\ No newline at end of file
+void InitGUI(const TCHAR* CommandLine, bool bP4EnvTabOnly = false);
\ No newline at end of file
diff --git a/Engine/Source/Programs/UnrealSync/Private/P4Env.cpp b/Engine/Source/Programs/UnrealSync/Private/P4Env.cpp
--- a/Engine/Source/Programs/UnrealSync/Private/P4Env.cpp
+++ b/Engine/Source/Programs/UnrealSync/Private/P4Env.cpp
@@ -5,6 +5,16 @@
 #include "P4Env.h"
 #include "ProcessHelper.h"
 #include "Regex.h"
+#include "Map.h"
+#include "JsonSerializer.h"
+#include "JsonObject.h"
+#include "JsonReader.h"
+#include "SEditableTextBox.h"
+#include "SGridPanel.h"
+
+DECLARE_LOG_CATEGORY_CLASS(LogP4Env, Log, All);
+
+#define LOCTEXT_NAMESPACE "UnrealSync"
 
 /**
  * Parses param from command line.
@@ -15,11 +25,156 @@
  *
  * @returns True if found, false otherwise.
  */
-bool GetParam(FString& Value, const TCHAR* CommandLine, const FString& ParamName)
+bool GetCommandLineParam(FString& Value, const TCHAR* CommandLine, EP4ParamType Type)
 {
-	return FParse::Value(CommandLine, *(ParamName + "="), Value);
+	return FParse::Value(CommandLine, *(FP4Env::GetParamName(Type) + "="), Value);
 }
 
+/**
+ * Cache class for file settings.
+ */
+class FSettingsCache
+{
+public:
+	/**
+	 * Instance getter.
+	 */
+	static FSettingsCache& Get()
+	{
+		static FSettingsCache Cache;
+		return Cache;
+	}
+
+	/**
+	 * Gets setting of given type.
+	 *
+	 * @param Value Output parameter. If succeeded this parameter will be overwritten with setting value.
+	 * @param Type Type to look for.
+	 *
+	 * @returns True if setting was found. False otherwise.
+	 */
+	bool GetSetting(FString& Value, EP4ParamType Type)
+	{
+		auto* ParamPtr = Settings.Find(FP4Env::GetParamName(Type));
+		if (ParamPtr == nullptr)
+		{
+			return false;
+		}
+
+		Value = *ParamPtr;
+		return true;
+	}
+
+	/**
+	 * Sets setting of given type to given value.
+	 *
+	 * @param Type Type of the setting.
+	 * @param Value Value of the setting.
+	 */
+	void SetSetting(EP4ParamType Type, const FString& Value)
+	{
+		auto* ParamPtr = Settings.Find(FP4Env::GetParamName(Type));
+		if (ParamPtr != nullptr)
+		{
+			*ParamPtr = Value;
+		}
+
+		Settings.Add(FP4Env::GetParamName(Type), Value);
+	}
+
+	/**
+	 * Saves setting to file.
+	 */
+	void Save()
+	{
+		TSharedRef<FJsonObject> Object(new FJsonObject());
+
+		for (const auto& Setting : Settings)
+		{
+			if (!Setting.Value.IsEmpty())
+			{
+				Object->SetStringField(Setting.Key, Setting.Value);
+			}
+		}
+
+		FString Buffer;
+		TSharedRef<TJsonWriter<> > Writer = TJsonWriterFactory<>::Create(&Buffer);
+		if (!FJsonSerializer::Serialize(Object, Writer))
+		{
+			UE_LOG(LogP4Env, Error, TEXT("Failed serializing settings using JSON."));
+			return;
+		}
+
+		auto bNeedsWrite = true;
+
+		if (FPaths::FileExists(GetSettingsFileName()))
+		{
+			// Read the file to a string.
+			FString FileContents;
+			bNeedsWrite = !(FFileHelper::LoadFileToString(FileContents, *GetSettingsFileName()) && FileContents == Buffer);
+		}
+
+		if (bNeedsWrite && !FFileHelper::SaveStringToFile(Buffer, *GetSettingsFileName()))
+		{
+			UE_LOG(LogP4Env, Error, TEXT("Failed writing settings to file: \"%s\"."), *GetSettingsFileName());
+			return;
+		}
+	}
+
+private:
+	/**
+	 * Constructor.
+	 *
+	 * Creates the cache from value read from settings file if it exists.
+	 */
+	FSettingsCache()
+	{
+		// Read the file to a string.
+		FString FileContents;
+		if (!FFileHelper::LoadFileToString(FileContents, *GetSettingsFileName()))
+		{
+			UE_LOG(LogP4Env, Log, TEXT("Couldn't find settings file \"%s\"."), *GetSettingsFileName());
+			return;
+		}
+
+		// Deserialize a JSON object from the string
+		TSharedPtr<FJsonObject> Object;
+		TSharedRef<TJsonReader<> > Reader = TJsonReaderFactory<>::Create(FileContents);
+		if (!FJsonSerializer::Deserialize(Reader, Object) || !Object.IsValid())
+		{
+			UE_LOG(LogP4Env, Error, TEXT("Settings file JSON parsing failed: \"%s\"."), *Reader->GetErrorMessage());
+			return;
+		}
+
+		for (const auto& Value : Object->Values)
+		{
+			FString Text;
+			if (!Value.Value->TryGetString(Text))
+			{
+				Settings.Empty();
+				break;
+			}
+
+			Settings.Add(Value.Key, Text);
+		}
+	}
+
+	/**
+	 * Gets settings file name. For now it's hardcoded.
+	 *
+	 * @returns Settings file name.
+	 */
+	static const FString& GetSettingsFileName()
+	{
+		const static FString HardCodedSettingsFileName = "UnrealSync.settings";
+
+		return HardCodedSettingsFileName;
+	}
+
+	/** Value map of settings. */
+	TMap<FString, FString> Settings;
+};
+
 /**
  * Gets param from environment variables.
  *
@@ -122,7 +277,7 @@ public:
 	 *
 	 * @returns Shared pointer to created iterator.
 	 */
-	static TSharedPtr<IP4EnvParamDetectionIterator> Create(EP4ParamType::Type Type, const TCHAR *CommandLine, const FP4Env& Env);
+	static TSharedPtr<IP4EnvParamDetectionIterator> Create(EP4ParamType Type, const TCHAR *CommandLine, const FP4Env& Env);
 };
 
 /**
@@ -137,12 +292,17 @@ public:
 	 * @param Type Type of the param to auto-detect.
 	 * @param CommandLine Command line to check for param.
 	 */
-	FP4EnvParamDetectionIteratorBase(EP4ParamType::Type Type, const TCHAR* CommandLine)
+	FP4EnvParamDetectionIteratorBase(EP4ParamType Type, const TCHAR* CommandLine)
 		: Type(Type)
 	{
-		if (GetParam(CommandLineParam, CommandLine, FP4Env::GetParamName(Type)))
+		if (GetCommandLineParam(HardParam, CommandLine, Type))
 		{
-			Finish();
+			OverriddenSetting(Type, HardParam, true);
+		}
+
+		if (FSettingsCache::Get().GetSetting(HardParam, Type))
+		{
+			OverriddenSetting(Type, HardParam, false);
 		}
 	}
 
@@ -155,10 +315,10 @@ public:
 	{
 		if (bFinished)
 		{
-			if (!CommandLineParam.IsEmpty())
+			if (!HardParam.IsEmpty())
 			{
-				Current = CommandLineParam;
-				CommandLineParam.Empty();
+				Current = HardParam;
+				HardParam.Empty();
 				return true;
 			}
 			else
@@ -208,15 +368,30 @@ protected:
 		bFinished = true;
 	}
 
+	/**
+	 * Method reports in the log that given param is overridden (from settings or command line).
+	 *
+	 * @param Type Param type to report on.
+	 * @param Value Value of this param.
+	 * @param CommandLine Boolean value that tells if this override comes from command line. If false it comes from settings file.
+	 */
+	void OverriddenSetting(EP4ParamType Type, const FString& Value, bool CommandLine)
+	{
+		FString Source = CommandLine ? TEXT("command line") : TEXT("settings file");
+
+		UE_LOG(LogP4Env, Log, TEXT("Setting %s overridden from %s and set to value \"%s\"."), *FP4Env::GetParamName(Type), *Source, *Value);
+		Finish();
+	}
+
 private:
 	/** Currently found param proposition. */
 	FString Current;
 
-	/** Found command line param. */
-	FString CommandLineParam;
+	/** Found hard param i.e. it can't be replaced by any autodetected ones. */
+	FString HardParam;
 
 	/** Type of the param. */
-	EP4ParamType::Type Type;
+	EP4ParamType Type;
 
 	/** Is this iterator finished. */
 	bool bFinished = false;
@@ -239,60 +414,80 @@ public:
 	FP4PathDetectionIterator(const TCHAR* CommandLine)
 		: Step(0), FP4EnvParamDetectionIteratorBase(EP4ParamType::Path, CommandLine)
 	{
-		PossibleLocations.Add("C:\\Program Files\\Perforce");
-		PossibleLocations.Add("C:\\Program Files (x86)\\Perforce");
-	}
+		const TCHAR* LocationsToLook[] =
+		{
+			TEXT("C:\\Program Files\\Perforce"),
+			TEXT("C:\\Program Files (x86)\\Perforce")
+		};
 
-	/**
-	 * Function that tries to detect P4 executable path.
-	 *
-	 * @param Output Output param. Found P4 executable path.
-	 *
-	 * @returns True if succeeded. False otherwise.
-	 */
-	virtual bool FindNext(FString& Output) override
-	{
 		// Tries to detect in standard environment paths.
 		static const FString WhereCommand = "where"; // TODO Mac: I think 'where' command equivalent on Mac is 'whereis'
 		static const FString P4ExecutableName = "p4.exe";
 
-		if (Step == 0)
+		FString WhereOutput;
+		if (RunProcessOutput(WhereCommand, P4ExecutableName, WhereOutput))
 		{
-			++Step;
+			AddUniqueLocation(FPaths::ConvertRelativePathToFull(
+				WhereOutput.Replace(TEXT("\n"), TEXT("")).Replace(TEXT("\r"), TEXT(""))));
+		}
 
-			FString WhereOutput;
-			if (RunProcessOutput(WhereCommand, P4ExecutableName, WhereOutput))
+		for (const auto* LocationToLook : LocationsToLook)
+		{
+			FString LocationCandidate = FPaths::ConvertRelativePathToFull(FPaths::Combine(LocationToLook, *P4ExecutableName));
+			if (FPaths::FileExists(LocationCandidate))
 			{
-				// P4 found, quick exit.
-				Output = FPaths::ConvertRelativePathToFull(
-					WhereOutput.Replace(TEXT("\n"), TEXT("")).Replace(TEXT("\r"), TEXT("")));
-				return true;
+				AddUniqueLocation(LocationCandidate);
 			}
-
-			return false;
 		}
+	}
 
-		if (Step > PossibleLocations.Num())
+	/**
+	 * Function that tries to detect P4 executable path.
+	 *
+	 * @param Output Output param. Found P4 executable path.
+	 *
+	 * @returns True if succeeded. False otherwise.
+	 */
+	virtual bool FindNext(FString& Output) override
+	{
+		if (Step >= PossibleLocations.Num())
 		{
 			Finish();
 			return false;
 		}
 
-		for (; Step <= PossibleLocations.Num(); ++Step)
+		Output = PossibleLocations[Step++];
+		return true;
+	}
+
+private:
+	/**
+	 * Adds unique location of P4.
+	 *
+	 * @param Path Location to add.
+	 */
+	void AddUniqueLocation(const FString& Path)
+	{
+		auto NormalizedPath = PlatformNormalize(Path);
+		if (!PossibleLocations.Contains(NormalizedPath))
 		{
-			FString LocationCandidate = FPaths::ConvertRelativePathToFull(FPaths::Combine(*PossibleLocations[Step - 1], *P4ExecutableName));
-			if (FPaths::FileExists(LocationCandidate))
-			{
-				Output = LocationCandidate;
-				return true;
-			}
+			PossibleLocations.Add(NormalizedPath);
 		}
+	}
 
-		Finish();
-		return false;
+	/**
+	 * Normalizes path in platform-specific way, i.e. lowers case on Windowsish
+	 * systems and leaves it intact on *nix.
+	 */
+	FString PlatformNormalize(const FString& Path)
+	{
+#if PLATFORM_WINDOWS || PLATFORM_WINRT || PLATFORM_XBOXONE
+		return Path.ToLower();
+#else
+		return Path;
+#endif
 	}
 
-private:
 	/** Possible P4 installation locations. */
 	TArray<FString> PossibleLocations;
 
@@ -518,8 +713,14 @@ public:
 	FP4ClientDetectionIterator(const TCHAR* CommandLine, const FP4Env& Env)
 		: FP4EnvParamDetectionIteratorBase(EP4ParamType::Client, CommandLine), Env(Env)
 	{
-		if (!RunProcessOutput(Env.GetPath(), FString::Printf(TEXT("-p%s clients -u%s"), *Env.GetPort(), *Env.GetUser()), P4ClientsOutput))
+		auto P4CommandLine = FString::Printf(TEXT("-p%s clients -u%s"), *Env.GetPort(), *Env.GetUser());
+		if (!RunProcessOutput(Env.GetPath(), P4CommandLine, P4ClientsOutput))
 		{
+			UE_LOG(LogP4Env, Log, TEXT("Failed to get client list. Used settings:") LINE_TERMINATOR
+				TEXT("\tPath: %s") LINE_TERMINATOR
+				TEXT("\tPort: %s") LINE_TERMINATOR
+				TEXT("\tUser: %s"), *Env.GetPath(), *Env.GetPort(), *Env.GetUser())
+
 			Finish();
 			return;
 		}
@@ -599,6 +800,16 @@ bool GetCurrentBranch(FString& Output, const FP4Env& Env)
 	if (!RunProcessOutput(Env.GetPath(), FString::Printf(TEXT("-p%s -u%s -c%s files %s"),
 		*Env.GetPort(), *Env.GetUser(), *Env.GetClient(), *GetKnownPath()), FilesOutput))
 	{
+		UE_LOG(LogP4Env, Log,
+			TEXT("Failed to verify branch. Details below:") LINE_TERMINATOR
+			TEXT("Note that you can't use this tool for cross syncing, i.e. you need to use") LINE_TERMINATOR
+			TEXT("UnrealSync from the same path as the branch being synced. In other word, if P4") LINE_TERMINATOR
+			TEXT("with given settings can't recognize currently running executable the") LINE_TERMINATOR
+			TEXT("verification will fail. Used settings:") LINE_TERMINATOR
+			TEXT("\tPath: %s") LINE_TERMINATOR
+			TEXT("\tPort: %s") LINE_TERMINATOR
+			TEXT("\tUser: %s") LINE_TERMINATOR
+			TEXT("\tClient: %s"), *Env.GetPath(),*Env.GetPort(), *Env.GetUser(), *Env.GetClient())
 		return false;
 	}
 
@@ -713,7 +924,7 @@ void FP4Env::SerializeParams(const FSerializationTask& SerializationTask)
 	SerializationTask.ExecuteIfBound(Branch, EP4ParamType::Branch);
 }
 
-FString FP4Env::GetParamName(EP4ParamType::Type Type)
+FString FP4Env::GetParamName(EP4ParamType Type)
 {
 	switch (Type)
 	{
@@ -737,7 +948,7 @@ bool FP4Env::AutoDetectMissingParams(const TCHAR* CommandLine)
 {
 	TArray<TSharedPtr<IP4EnvParamDetectionIterator> > ParamDetectionIteratorsStack;
 
-	EP4ParamType::Type Type = EP4ParamType::Path;
+	auto Type = EP4ParamType::Path;
 	ParamDetectionIteratorsStack.Add(IP4EnvParamDetectionIterator::Create(Type, CommandLine, *this));
 
 	int32 IterationCountdown = 20;
@@ -750,7 +961,7 @@ bool FP4Env::AutoDetectMissingParams(const TCHAR* CommandLine)
 
 			if (Type != EP4ParamType::Branch)
 			{
-				Type = (EP4ParamType::Type) ((int)Type + 1);
+				Type = (EP4ParamType) ((int)Type + 1);
 				ParamDetectionIteratorsStack.Add(IP4EnvParamDetectionIterator::Create(Type, CommandLine, *this));
 				continue;
 			}
@@ -761,7 +972,7 @@ bool FP4Env::AutoDetectMissingParams(const TCHAR* CommandLine)
 		}
 		else
 		{
-			Type = (EP4ParamType::Type) ((int)Type - 1);
+			Type = (EP4ParamType) ((int)Type - 1);
 			ParamDetectionIteratorsStack.RemoveAt(ParamDetectionIteratorsStack.Num() - 1);
 			--IterationCountdown;
 			if (!IterationCountdown)
@@ -779,7 +990,7 @@ FString FP4Env::GetCommandLine()
 	class CommandLineOutput
 	{
 	public:
-		void AppendParam(FString& FieldReference, EP4ParamType::Type Type)
+		void AppendParam(FString& FieldReference, EP4ParamType Type)
 		{
 			Output += " -" + GetParamName(Type) + "=\"" + FieldReference + "\"";
 		}
@@ -799,7 +1010,7 @@ const FString& FP4Env::GetPath() const
 	return Path;
 }
 
-void FP4Env::SetParam(EP4ParamType::Type Type, const FString& Value)
+void FP4Env::SetParam(EP4ParamType Type, const FString& Value)
 {
 	switch (Type)
 	{
@@ -824,7 +1035,74 @@ void FP4Env::SetParam(EP4ParamType::Type Type, const FString& Value)
 	}
 }
 
-TSharedPtr<IP4EnvParamDetectionIterator> IP4EnvParamDetectionIterator::Create(EP4ParamType::Type Type, const TCHAR* CommandLine, const FP4Env& Env)
+bool FP4Env::CheckIfFileNeedsUpdate(const FString& FilePath)
+{
+	FString Output;
+	if (!RunP4Output(FString::Printf(TEXT("fstat %s"), *FilePath), Output))
+	{
+		UE_LOG(LogP4Env, Error, TEXT("Checking if file \"%s\" needs update, failed."), *FilePath);
+		return false;
+	}
+
+	static const FRegexPattern HeadRevPattern(TEXT("\\.\\.\\. headRev (\\d+)"));
+	static const FRegexPattern HaveRevPattern(TEXT("\\.\\.\\. haveRev (\\d+)"));
+	static const FRegexPattern ChangePattern(TEXT("\\.\\.\\. change ([^\\n]+)"));
+
+	FRegexMatcher ChangeMatcher(ChangePattern, Output);
+
+	if (ChangeMatcher.FindNext())
+	{
+		UE_LOG(LogP4Env, Error, TEXT("File \"%s\" is checked out, can't update."), *FilePath);
+		return false;
+	}
+
+	FRegexMatcher HeadRevMatcher(HeadRevPattern, Output);
+	FRegexMatcher HaveRevMatcher(HaveRevPattern, Output);
+
+	if (!HeadRevMatcher.FindNext() || !HaveRevMatcher.FindNext())
+	{
+		UE_LOG(LogP4Env, Error, TEXT("Can't determine head or have revision for file \"%s\"."), *FilePath);
+		return false;
+	}
+
+	int32 HeadRev = FPlatformString::Atoi(*Output.Mid(
+		HeadRevMatcher.GetCaptureGroupBeginning(1),
+		HeadRevMatcher.GetCaptureGroupEnding(1) - HeadRevMatcher.GetCaptureGroupBeginning(1)));
+
+	int32 HaveRev = FPlatformString::Atoi(*Output.Mid(
+		HaveRevMatcher.GetCaptureGroupBeginning(1),
+		HaveRevMatcher.GetCaptureGroupEnding(1) - HaveRevMatcher.GetCaptureGroupBeginning(1)));
+
+	return HaveRev < HeadRev;
+}
+
+bool FP4Env::IsValid()
+{
+	return Env.IsValid();
+}
+
+const FString& FP4Env::GetParamByType(EP4ParamType Type) const
+{
+	switch (Type)
+	{
+	case EP4ParamType::Path:
+		return GetPath();
+	case EP4ParamType::Port:
+		return GetPort();
+	case EP4ParamType::User:
+		return GetUser();
+	case EP4ParamType::Client:
+		return GetClient();
+	case EP4ParamType::Branch:
+		return GetBranch();
+	default:
+		// Unimplemented param type?
+		checkNoEntry();
+		throw 0; // Won't get here, used to suppress "not all control paths return a value"
+	}
+}
+
+TSharedPtr<IP4EnvParamDetectionIterator> IP4EnvParamDetectionIterator::Create(EP4ParamType Type, const TCHAR* CommandLine, const FP4Env& Env)
 {
 	switch (Type)
 	{
@@ -846,4 +1124,284 @@ TSharedPtr<IP4EnvParamDetectionIterator> IP4EnvParamDetectionIterator::Create(EP
 }
 
 /* Static variable initialization. */
-TSharedPtr<FP4Env> FP4Env::Env;
\ No newline at end of file
+TSharedPtr<FP4Env> FP4Env::Env;
+
+/**
+ * Grid panel widget with labelled options.
+ */
+class SOptionsPanel : public SGridPanel
+{
+public:
+	SLATE_BEGIN_ARGS(SOptionsPanel) { }
+	SLATE_END_ARGS()
+
+	/**
+	 * Method that constructs this widget given Slate construction arguments.
+	 *
+	 * @param InArgs Slate construction arguments.
+	 */
+	void Construct(const FArguments& InArgs)
+	{
+		SetColumnFill(1, 1.0f);
+	}
+
+	/**
+	 * Inner class forward declaration.
+	 */
+	class SOption;
+
+	/**
+	 * Adds option widget to the grid.
+	 *
+	 * @param Option Option to add.
+	 */
+	void Add(TSharedPtr<SOption> Option);
+
+private:
+	/** Collection of all options in this grid. */
+	TArray<TSharedPtr<SOption> > Options;
+};
+
+/**
+ * Widget and label provider for any option added.
+ */
+class SOptionsPanel::SOption : public TSharedFromThis<SOption>
+{
+	friend void SOptionsPanel::Add(TSharedPtr<SOption> Option);
+
+public:
+	/**
+	 * Constructor
+	 */
+	SOption() { }
+
+protected:
+	/**
+	 * Must be overridden. Provides a label for this option.
+	 *
+	 * @returns Label for this option.
+	 */
+	virtual FText GetLabel() const = 0;
+
+	/**
+	 * Creates and returns widget for this option. If not overridden it creates editable text box.
+	 *
+	 * @returns Widget for this option.
+	 */
+	virtual TSharedRef<SWidget> CreateWidget() { return SNew(SEditableTextBox); }
+};
+
+void SOptionsPanel::Add(TSharedPtr<SOption> Option)
+{
+	auto Id = Options.Add(Option);
+
+	AddSlot(0, Id).Padding(1.0f)
+		[
+			SNew(SHorizontalBox)
+			+ SHorizontalBox::Slot()
+			+ SHorizontalBox::Slot().AutoWidth().VAlign(VAlign_Center)
+			[
+				SNew(STextBlock).Text(Option.Get(), &SOption::GetLabel)
+			]
+		];
+
+	AddSlot(1, Id).Padding(1.0f)
+		[
+			Option->CreateWidget()
+		];
+}
+
+BEGIN_SLATE_FUNCTION_BUILD_OPTIMIZATION
+void SP4EnvTabWidget::Construct(const FArguments& InArgs)
+{
+	/**
+	 * P4 option widget.
+	 */
+	class SP4Option : public SOptionsPanel::SOption
+	{
+	public:
+		/**
+		 * Constructor
+		 *
+		 * @param Reference to FP4Option object this object relates to.
+		 */
+		SP4Option(TSharedRef<FP4Option> Option)
+			: Option(Option)
+		{
+
+		}
+
+	protected:
+		/**
+		 * Creates and returns widget for this option.
+		 *
+		 * @returns Widget for this option.
+		 */
+		virtual TSharedRef<SWidget> CreateWidget() override
+		{
+			return SNew(SEditableTextBox)
+				.OnTextCommitted(this, &SP4Option::OnTextCommitted)
+				.Text(this, &SP4Option::GetText)
+				.ToolTipText(FText::Format(LOCTEXT("P4OptionToolTip", "This is a setting for {0}. If you see greyed text, then it means that this setting was autodetected."), FText::FromString(FP4Env::GetParamName(Option->GetType()))))
+				.HintText(this, &SP4Option::GetHintText);
+		}
+
+		/**
+		 * On text committed event handler.
+		 *
+		 * @param CommittedText Committed text.
+		 * @param Type Commit type.
+		 */
+		void OnTextCommitted(const FText& CommittedText, ETextCommit::Type Type)
+		{
+			Option->GetText() = CommittedText;
+		}
+
+		/**
+		 * SEditableTextBox text provider.
+		 *
+		 * @returns Text that should be displayed in SEditableTextBox for this option.
+		 */
+		FText GetText() const
+		{
+			return Option->GetText();
+		}
+
+		/**
+		 * SEditableTextBox hint text provider (greyed text).
+		 *
+		 * @returns Hint text that should be displayed in SEditableTextBox for this option.
+		 */
+		FText GetHintText() const
+		{
+			if (FP4Env::IsValid())
+			{
+				auto Env = FP4Env::Get();
+
+				return FText::FromString(Env.GetParamByType(Option->GetType()));
+			}
+
+			return FText();
+		}
+
+		/**
+		 * Provides a label for this option.
+		 *
+		 * @returns Label for this option.
+		 */
+		virtual FText GetLabel() const override
+		{
+			switch (Option->GetType())
+			{
+			case EP4ParamType::Branch:
+				return LOCTEXT("P4Option_Branch", "Branch name");
+			case EP4ParamType::Client:
+				return LOCTEXT("P4Option_Client", "Workspace name");
+			case EP4ParamType::Path:
+				return LOCTEXT("P4Option_Path", "Path to P4 executable");
+			case EP4ParamType::Port:
+				return LOCTEXT("P4Option_Port", "P4 server address");
+			case EP4ParamType::User:
+				return LOCTEXT("P4Option_User", "P4 user name");
+			default:
+				checkNoEntry();
+				throw 0; // Won't get here, used to suppress "not all control paths return a value"
+			}
+		}
+
+	private:
+		/** Option reference. */
+		TSharedRef<FP4Option> Option;
+	};
+
+	Options.Reserve(4);
+	Options.Add(MakeShareable(new FP4Option(EP4ParamType::Path)));
+	Options.Add(MakeShareable(new FP4Option(EP4ParamType::Port)));
+	Options.Add(MakeShareable(new FP4Option(EP4ParamType::User)));
+	Options.Add(MakeShareable(new FP4Option(EP4ParamType::Client)));
+
+	auto OptionsPanel = SNew(SOptionsPanel);
+
+	for (auto Option : Options)
+	{
+		OptionsPanel->Add(MakeShareable(new SP4Option(Option)));
+	}
+
+	this->ChildSlot
+		[
+			SNew(SVerticalBox)
+			+ SVerticalBox::Slot().HAlign(HAlign_Center).AutoHeight().Padding(10.0f)
+			[
+				SNew(STextBlock)
+				.Font(FSlateFontInfo(FPaths::EngineContentDir() / TEXT("Slate/Fonts/Roboto-Regular.ttf"), 14))
+				.Text(LOCTEXT("P4Settings", "Perforce settings"))
+			]
+			+ SVerticalBox::Slot()
+			[
+				OptionsPanel
+			]
+			+ SVerticalBox::Slot().AutoHeight()
+			[
+				SNew(SHorizontalBox)
+				+ SHorizontalBox::Slot().VAlign(VAlign_Center)
+				[
+					SNew(STextBlock).Text(
+						!FP4Env::IsValid()
+						? LOCTEXT("P4AutodetectFailed", "P4 environment detection failed. See the log for more details.")
+						: LOCTEXT("P4AutodetectSucceeded", "P4 environment detection succeeded.")
+					)
+				]
+				+ SHorizontalBox::Slot().AutoWidth()
+				[
+					SNew(SButton)
+					.OnClicked(this, &SP4EnvTabWidget::OnSaveAndRestartButtonClick)
+					[
+						SNew(STextBlock).Text(LOCTEXT("SaveAndRestart", "Save and restart"))
+					]
+				]
+				+ SHorizontalBox::Slot().AutoWidth()
+				[
+					SNew(SButton)
+					.OnClicked(this, &SP4EnvTabWidget::OnCloseButtonClick)
+					[
+						SNew(STextBlock).Text(LOCTEXT("Close", "Close"))
+					]
+				]
+			]
+		];
+}
+END_SLATE_FUNCTION_BUILD_OPTIMIZATION
+
+FReply SP4EnvTabWidget::OnSaveAndRestartButtonClick()
+{
+	auto& Settings = FSettingsCache::Get();
+
+	for (auto Option : Options)
+	{
+		Settings.SetSetting(Option->GetType(), Option->GetText().ToString());
+	}
+
+	Settings.Save();
+
+	FUnrealSync::RunDetachedUS(FPlatformProcess::ExecutableName(false), true, true, false);
+	FPlatformMisc::RequestExit(false);
+
+	return FReply::Handled();
+}
+
+FReply SP4EnvTabWidget::OnCloseButtonClick()
+{
+	FPlatformMisc::RequestExit(false);
+
+	return FReply::Handled();
+}
+
+SP4EnvTabWidget::FP4Option::FP4Option(EP4ParamType Type)
+	: Type(Type)
+{
+	FString Param;
+	if (FSettingsCache::Get().GetSetting(Param, Type))
+	{
+		Text = FText::FromString(Param);
+	}
+}
diff --git a/Engine/Source/Programs/UnrealSync/Private/P4Env.h b/Engine/Source/Programs/UnrealSync/Private/P4Env.h
--- a/Engine/Source/Programs/UnrealSync/Private/P4Env.h
+++ b/Engine/Source/Programs/UnrealSync/Private/P4Env.h
@@ -8,17 +8,87 @@
 /**
  * Enum describing P4 param type.
  */
-namespace EP4ParamType
+enum class EP4ParamType
 {
-	enum Type
+	Path = 0,
+	Port = 1,
+	User = 2,
+	Client = 3,
+	Branch = 4
+};
+
+/**
+ * P4 settings tab.
+ */
+class SP4EnvTabWidget : public SCompoundWidget
+{
+public:
+	SLATE_BEGIN_ARGS(SP4EnvTabWidget){}
+	SLATE_END_ARGS()
+
+	/**
+	 * Slate constructor
+	 *
+	 * @param InArgs Slate args.
+	 */
+	void Construct(const FArguments& InArgs);
+
+	/**
+	 * Function that is called on when Close button is clicked.
+	 * It tells the app to close.
+	 *
+	 * @returns Tells that this event was handled.
+	 */
+	FReply OnCloseButtonClick();
+
+	/**
+	 * Function that is called on when Close button is clicked.
+	 * It tells the app to close.
+	 *
+	 * @returns Tells that this event was handled.
+	 */
+	FReply OnSaveAndRestartButtonClick();
+
+private:
+	/**
+	 * P4 option container.
+	 */
+	class FP4Option
 	{
-		Path	= 0,
-		Port	= 1,
-		User	= 2,
-		Client	= 3,
-		Branch	= 4
+	public:
+		/**
+		 * Constructor
+		 *
+		 * @param Type The type for this option.
+		 */
+		FP4Option(EP4ParamType Type);
+
+		/**
+		 * Gets the type of this option.
+		 */
+		EP4ParamType GetType() const { return Type; }
+
+		/**
+		 * Gets the text of this option.
+		 */
+		FText GetText() const { return Text; }
+
+		/**
+		 * Gets the reference to text of this option.
+		 */
+		FText& GetText() { return Text; }
+
+	private:
+		/** This option's type. */
+		EP4ParamType Type;
+
+		/** This option's text. */
+		FText Text;
 	};
-}
+
+	/** Collection of options maintained by this widget. */
+	TArray<TSharedRef<FP4Option> > Options;
+};
 
 /**
  * P4 environment class to connect with P4 executable.
@@ -57,6 +127,13 @@ public:
 	 */
 	static FP4Env& Get();
 
+	/**
+	 * Tells if P4 environment is valid.
+	 *
+	 * @returns True if P4 environment is valid. False otherwise.
+	 */
+	static bool IsValid();
+
 	/**
 	 * Runs P4 process and allows communication through progress delegate.
 	 *
@@ -137,11 +214,29 @@ public:
 	 *
 	 * @returns Param type name.
 	 */
-	static FString GetParamName(EP4ParamType::Type Type);
+	static FString GetParamName(EP4ParamType Type);
+
+	/**
+	 * Gets param value using type.
+	 *
+	 * @param Type Param type enum id.
+	 *
+	 * @returns Param value.
+	 */
+	const FString& GetParamByType(EP4ParamType Type) const;
+
+	/**
+	 * Checks if given file needs update, i.e. is not at head revision.
+	 *
+	 * @param FilePath
+	 *
+	 * @returns True if needs update. False otherwise.
+	 */
+	static bool CheckIfFileNeedsUpdate(const FString& FilePath);
 
 private:
 	/* Param serialization delegate. */
-	DECLARE_DELEGATE_TwoParams(FSerializationTask, FString&, EP4ParamType::Type);
+	DECLARE_DELEGATE_TwoParams(FSerializationTask, FString&, EP4ParamType);
 
 	/**
 	 * Constructor
@@ -170,7 +265,7 @@ private:
 	 * @param Type Type of param to set.
 	 * @param Value Value to set.
 	 */
-	void SetParam(EP4ParamType::Type Type, const FString& Value);
+	void SetParam(EP4ParamType Type, const FString& Value);
 
 	/* Singleton instance pointer. */
 	static TSharedPtr<FP4Env> Env;
diff --git a/Engine/Source/Programs/UnrealSync/Private/UnrealSync.cpp b/Engine/Source/Programs/UnrealSync/Private/UnrealSync.cpp
--- a/Engine/Source/Programs/UnrealSync/Private/UnrealSync.cpp
+++ b/Engine/Source/Programs/UnrealSync/Private/UnrealSync.cpp
@@ -11,14 +11,56 @@ DEFINE_LOG_CATEGORY_STATIC(LogUnrealSync, Log, All);
 
 IMPLEMENT_APPLICATION(UnrealSync, "UnrealSync");
 
-/**
- * Unreal sync main function.
- *
- * @param CommandLine Command line that the program was called with.
- */
-void RunUnrealSync(const TCHAR* CommandLine)
+bool FUnrealSync::UpdateOriginalUS(const FString& OriginalUSPath)
+{
+	FString Output;
+	return FP4Env::RunP4Output(TEXT("sync ") + OriginalUSPath, Output);
+}
+
+bool FUnrealSync::RunDetachedUS(const FString& USPath, bool bDoNotRunFromCopy, bool bDoNotUpdateOnStartUp, bool bPassP4Env)
+{
+	FString CommandLine = FString()
+		+ (bDoNotRunFromCopy ? TEXT("-DoNotRunFromCopy ") : TEXT(""))
+		+ (bDoNotUpdateOnStartUp ? TEXT("-DoNotUpdateOnStartUp ") : TEXT(""))
+		+ (bPassP4Env ? *FP4Env::Get().GetCommandLine()	: TEXT(""));
+
+	return RunProcess(USPath, CommandLine);
+}
+
+bool FUnrealSync::DeleteIfExistsAndCopyFile(const FString& To, const FString& From)
 {
-	InitGUI(CommandLine);
+	auto& PlatformPhysical = IPlatformFile::GetPlatformPhysical();
+
+	if (PlatformPhysical.FileExists(*To) && !PlatformPhysical.DeleteFile(*To))
+	{
+		return false;
+	}
+
+	return PlatformPhysical.CopyFile(*To, *From);
+}
+
+void FUnrealSync::RunUnrealSync(const TCHAR* CommandLine)
+{
+	// start up the main loop
+	GEngineLoop.PreInit(CommandLine);
+
+	if (!FP4Env::Init(CommandLine))
+	{
+		GLog->Flush();
+		InitGUI(CommandLine, true);
+	}
+	else
+	{
+		if (FUnrealSync::Initialization(CommandLine))
+		{
+			InitGUI(CommandLine);
+		}
+
+		if (!FUnrealSync::GetInitializationError().IsEmpty())
+		{
+			UE_LOG(LogUnrealSync, Fatal, TEXT("Error: %s"), *FUnrealSync::GetInitializationError());
+		}
+	}
 }
 
 FString FUnrealSync::GetLatestLabelForGame(const FString& GameName)
@@ -435,16 +477,66 @@ bool FUnrealSync::Sync(bool bArtist, bool bPreview, const FString& Label, const
 
 	if (!FP4Env::RunP4Progress(FString("sync ") + (bPreview ? "-n " : "") + FString::Join(SyncSteps, TEXT(" ")), OnSyncProgress))
 	{
+	}
+
+	return true;
+}
+
+bool FUnrealSync::Initialization(const TCHAR* CommandLine)
+{
+	FString CommonExecutablePath =
+		FPaths::ConvertRelativePathToFull(
+		FPaths::Combine(*FPaths::EngineDir(), TEXT("Binaries"), TEXT("Win64"),
+#if !UE_BUILD_DEBUG
+		TEXT("UnrealSync")
+#else
+		TEXT("UnrealSync-Win64-Debug")
+#endif
+		));
+
+	FString OriginalExecutablePath = CommonExecutablePath + TEXT(".exe");
+	FString TemporaryExecutablePath = CommonExecutablePath + TEXT(".Temporary.exe");
+
+	bool bDoNotRunFromCopy = FParse::Param(CommandLine, TEXT("DoNotRunFromCopy"));
+	bool bDoNotUpdateOnStartUp = FParse::Param(CommandLine, TEXT("DoNotUpdateOnStartUp"));
+
+	if (!bDoNotRunFromCopy)
+	{
+		if (!DeleteIfExistsAndCopyFile(*TemporaryExecutablePath, *OriginalExecutablePath))
+		{
+			InitializationError = TEXT("Copying UnrealSync to temp location failed.");
+		}
+		else if (!RunDetachedUS(TemporaryExecutablePath, true, bDoNotUpdateOnStartUp, true))
+		{
+			InitializationError = TEXT("Running remote UnrealSync failed.");
+		}
+
+		return false;
+	}
+
+	if (!bDoNotUpdateOnStartUp && FP4Env::CheckIfFileNeedsUpdate(OriginalExecutablePath))
+	{
+		if (!UpdateOriginalUS(OriginalExecutablePath))
+		{
+			InitializationError = TEXT("UnrealSync update failed.");
+		}
+		else if (!RunDetachedUS(OriginalExecutablePath, false, true, true))
+		{
+			InitializationError = TEXT("Running UnrealSync failed.");
+		}
+
 		return false;
 	}
 
 	return true;
 }
 
+
 /* Static fields initialization. */
 FUnrealSync::FOnDataLoaded FUnrealSync::OnDataLoaded;
 FUnrealSync::FOnDataReset FUnrealSync::OnDataReset;
 TSharedPtr<FP4DataCache> FUnrealSync::Data;
 bool FUnrealSync::bLoadingFinished = false;
 TSharedPtr<FP4DataLoader> FUnrealSync::LoaderThread;
-TSharedPtr<FSyncingThread> FUnrealSync::SyncingThread;
\ No newline at end of file
+TSharedPtr<FSyncingThread> FUnrealSync::SyncingThread;
+FString FUnrealSync::InitializationError;
\ No newline at end of file
diff --git a/Engine/Source/Programs/UnrealSync/Private/UnrealSync.h b/Engine/Source/Programs/UnrealSync/Private/UnrealSync.h
--- a/Engine/Source/Programs/UnrealSync/Private/UnrealSync.h
+++ b/Engine/Source/Programs/UnrealSync/Private/UnrealSync.h
@@ -6,11 +6,6 @@
 #include "CoreUObject.h"
 #include "P4DataCache.h"
 
-/**
- * Main program function.
- */
-void RunUnrealSync(const TCHAR* CommandLine);
-
 /**
  * Class to store date needed for sync monitoring thread.
  */
@@ -104,6 +99,16 @@ public:
 	/* On sync log chunk read event delegate. */
 	DECLARE_DELEGATE_RetVal_OneParam(bool, FOnSyncProgress, const FString&);
 
+	/**
+	 * This method copies UnrealSync to temp location and run it from there,
+	 * and updates (if needed).
+	 *
+	 * @param CommandLine Command line.
+	 *
+	 * @returns True if initialization phase is over and UnrealSync can build-up GUI. False otherwise.
+	 */
+	static bool Initialization(const TCHAR* CommandLine);
+
 	/**
 	 * Gets latest label for given game name.
 	 *
@@ -171,6 +176,27 @@ public:
 	 */
 	static void StartLoadingData();
 
+	/**
+	 * Runs detached UnrealSync process and passes given parameters in command line.
+	 *
+	 * @param USPath UnrealSync executable path.
+	 * @param bDoNotRunFromCopy Should UnrealSync be called with -DoNotRunFromCopy flag?
+	 * @param bDoNotUpdateOnStartUp Should UnrealSync be called with -DoNotUpdateOnStartUp flag?
+	 * @param bPassP4Env Pass P4 environment parameters to UnrealSync?
+	 *
+	 * @returns True if succeeded. False otherwise. Notice that this says of
+	 *			success of launching procedure not launched process, cause
+	 *			its detached.
+	 */
+	static bool RunDetachedUS(const FString& USPath, bool bDoNotRunFromCopy, bool bDoNotUpdateOnStartUp, bool bPassP4Env);
+
+	/**
+	 * Unreal sync main function.
+	 *
+	 * @param CommandLine Command line that the program was called with.
+	 */
+	static void RunUnrealSync(const TCHAR* CommandLine);
+
 	/**
 	 * Tells that labels names are currently being loaded.
 	 *
@@ -207,6 +233,13 @@ public:
 	 */
 	static bool LoadingFinished();
 
+	/**
+	 * Gets initialization error.
+	 *
+	 * @returns Initialization error.
+	 */
+	static const FString& GetInitializationError() { return InitializationError; }
+
 	/**
 	 * Gets labels from the loaded cache.
 	 *
@@ -236,6 +269,26 @@ public:
 	 */
 	static bool Sync(bool bArtist, bool bPreview, const FString& Label, const FString& Game, const FOnSyncProgress& OnSyncProgress);
 private:
+	/**
+	 * Tries to update original UnrealSync at given location.
+	 *
+	 * @param Location of original UnrealSync.
+	 *
+	 * @returns False on failure. True otherwise.
+	 */
+	static bool UpdateOriginalUS(const FString& OriginalUSPath);
+
+	/**
+	 * This function copies file from From to To location and deletes
+	 * if To exists.
+	 *
+	 * @param To Location to which copy the file.
+	 * @param From Location from which copy the file.
+	 *
+	 * @returns True if succeeded. False otherwise.
+	 */
+	static bool DeleteIfExistsAndCopyFile(const FString& To, const FString& From);
+
 	/* Tells if loading has finished. */
 	static bool bLoadingFinished;
 
@@ -253,4 +306,7 @@ private:
 
 	/* Background syncing process monitoring thread. */
 	static TSharedPtr<FSyncingThread> SyncingThread;
+
+	/* Error that happened during application initialization. */
+	static FString InitializationError;
 };
diff --git a/Engine/Source/Programs/UnrealSync/Private/Windows/UnrealSyncMainWindows.cpp b/Engine/Source/Programs/UnrealSync/Private/Windows/UnrealSyncMainWindows.cpp
--- a/Engine/Source/Programs/UnrealSync/Private/Windows/UnrealSyncMainWindows.cpp
+++ b/Engine/Source/Programs/UnrealSync/Private/Windows/UnrealSyncMainWindows.cpp
@@ -7,7 +7,7 @@
  */
 int WINAPI WinMain(_In_ HINSTANCE hInInstance, _In_opt_ HINSTANCE hPrevInstance, _In_ LPSTR, _In_ int nCmdShow)
 {
-	RunUnrealSync(GetCommandLineW());
+	FUnrealSync::RunUnrealSync(GetCommandLineW());
 
 	return 0;
 }
\ No newline at end of file
diff --git a/Engine/Source/Programs/UnrealSync/UnrealSync.Build.cs b/Engine/Source/Programs/UnrealSync/UnrealSync.Build.cs
--- a/Engine/Source/Programs/UnrealSync/UnrealSync.Build.cs
+++ b/Engine/Source/Programs/UnrealSync/UnrealSync.Build.cs
@@ -18,7 +18,8 @@ public class UnrealSync : ModuleRules
 				"SlateCore",
 				"SlateReflector",
 				"StandaloneRenderer",
-				"XmlParser"
+				"XmlParser",
+				"JSON"
 			}
 		);
 		
