UnrealSync: Fail if P4 sync failed.

diff --git a/Engine/Source/Programs/UnrealSync/Private/UnrealSync.cpp b/Engine/Source/Programs/UnrealSync/Private/UnrealSync.cpp
--- a/Engine/Source/Programs/UnrealSync/Private/UnrealSync.cpp
+++ b/Engine/Source/Programs/UnrealSync/Private/UnrealSync.cpp
@@ -477,6 +477,7 @@ bool FUnrealSync::Sync(bool bArtist, bool bPreview, const FString& Label, const
 
 	if (!FP4Env::RunP4Progress(FString("sync ") + (bPreview ? "-n " : "") + FString::Join(SyncSteps, TEXT(" ")), OnSyncProgress))
 	{
+		return false;
 	}
 
 	return true;
