UnrealSync: Added ability to override sync steps to custom one during debug session.

diff --git a/Engine/Source/Programs/UnrealSync/Private/GUI.cpp b/Engine/Source/Programs/UnrealSync/Private/GUI.cpp
--- a/Engine/Source/Programs/UnrealSync/Private/GUI.cpp
+++ b/Engine/Source/Programs/UnrealSync/Private/GUI.cpp
@@ -789,8 +789,17 @@ public:
 				SNew(STextBlock).Text(FText::FromString("Artist sync?"))
 			];
 
-		// Checked by default.
-		ArtistSyncCheckBox->ToggleCheckedState();
+#if UE_BUILD_DEBUG
+		auto bDebug = FParse::Param(FCommandLine::Get(), TEXT("Debug"));
+#else
+		const auto bDebug = false;
+#endif
+
+		if (!bDebug)
+		{
+			// Checked by default.
+			ArtistSyncCheckBox->ToggleCheckedState();
+		}
 
 		PreviewSyncCheckBox = SNew(SCheckBox)
 			[
@@ -807,11 +816,13 @@ public:
 		LogListView = SNew(SListView<TSharedPtr<FString> >)
 			.ListItemsSource(&LogLines)
 			.OnGenerateRow(this, &SMainTabWidget::GenerateLogItem);
+
+		TSharedPtr<SVerticalBox> MainBox;
 		
 		Switcher = SNew(SWidgetSwitcher)
 			+ SWidgetSwitcher::Slot()
 			[
-				SNew(SVerticalBox)
+				SAssignNew(MainBox, SVerticalBox)
 				+ SVerticalBox::Slot().AutoHeight().Padding(5.0f)
 				[
 					PickGameWidget.ToSharedRef()
@@ -881,6 +892,22 @@ public:
 				Switcher.ToSharedRef()
 			];
 
+		if (bDebug)
+		{
+			MainBox->InsertSlot(2).AutoHeight()
+				[
+					SNew(SHorizontalBox)
+					+ SHorizontalBox::Slot()
+					[
+						SNew(STextBlock).Text(LOCTEXT("OverrideSyncStepSpec", "Override sync step specification"))
+					]
+					+ SHorizontalBox::Slot()
+					[
+						SNew(SEditableTextBox).OnTextCommitted(this, &SMainTabWidget::OnOverrideSyncStepText)
+					]
+				];
+		}
+
 		FUnrealSync::RegisterOnDataReset(FUnrealSync::FOnDataReset::CreateRaw(this, &SMainTabWidget::DataReset));
 		FUnrealSync::RegisterOnDataLoaded(FUnrealSync::FOnDataLoaded::CreateRaw(this, &SMainTabWidget::DataLoaded));
 
@@ -1015,11 +1042,13 @@ private:
 	 */
 	FReply OnSync()
 	{
-		auto bArtist = ArtistSyncCheckBox->IsChecked();
-		auto bPreview = PreviewSyncCheckBox->IsChecked();
+		FSyncSettings Settings(
+			ArtistSyncCheckBox->IsChecked(),
+			PreviewSyncCheckBox->IsChecked(),
+			OverrideSyncStep);
 
 		Switcher->SetActiveWidgetIndex(1);
-		FUnrealSync::LaunchSync(bArtist, bPreview, GetCurrentSyncCmdLineProvider(),
+		FUnrealSync::LaunchSync(Settings, GetCurrentSyncCmdLineProvider(),
 			FUnrealSync::FOnSyncFinished::CreateRaw(this, &SMainTabWidget::SyncingFinished),
 			FUnrealSync::FOnSyncProgress::CreateRaw(this, &SMainTabWidget::SyncingProgress));
 
@@ -1125,6 +1154,11 @@ private:
 		}
 	}
 
+	/**
+	 * Happens when someone edits the override sync step text box.
+	 */
+	void OnOverrideSyncStepText(const FText& CommittedText, ETextCommit::Type Type);
+
 	/**
 	 * Method to add command line provider/widget to correct arrays.
 	 *
@@ -1164,8 +1198,16 @@ private:
 
 	/* Array of sync command line providers. */
 	TArray<TSharedRef<ILabelNameProvider> > SyncCommandLineProviders;
+
+	/* Override sync step value. */
+	FString OverrideSyncStep;
 };
 
+void SMainTabWidget::OnOverrideSyncStepText(const FText& CommittedText, ETextCommit::Type Type)
+{
+	OverrideSyncStep = CommittedText.ToString();
+}
+
 /**
  * Creates and returns main tab.
  *
diff --git a/Engine/Source/Programs/UnrealSync/Private/P4Env.cpp b/Engine/Source/Programs/UnrealSync/Private/P4Env.cpp
--- a/Engine/Source/Programs/UnrealSync/Private/P4Env.cpp
+++ b/Engine/Source/Programs/UnrealSync/Private/P4Env.cpp
@@ -14,8 +14,6 @@
 
 DECLARE_LOG_CATEGORY_CLASS(LogP4Env, Log, All);
 
-#define LOCTEXT_NAMESPACE "UnrealSync"
-
 /**
  * Parses param from command line.
  *
diff --git a/Engine/Source/Programs/UnrealSync/Private/UnrealSync.h b/Engine/Source/Programs/UnrealSync/Private/UnrealSync.h
--- a/Engine/Source/Programs/UnrealSync/Private/UnrealSync.h
+++ b/Engine/Source/Programs/UnrealSync/Private/UnrealSync.h
@@ -6,6 +6,8 @@
 #include "CoreUObject.h"
 #include "P4DataCache.h"
 
+#define LOCTEXT_NAMESPACE "UnrealSync"
+
 /**
  * Class to store date needed for sync monitoring thread.
  */
@@ -82,6 +84,30 @@ private:
 	FString CurrentGameName;
 };
 
+struct FSyncSettings
+{
+	/* Artist sync? */
+	bool bArtist;
+
+	/* Preview sync? */
+	bool bPreview;
+
+	/* Override sync step. If set then it overrides collecting sync steps and uses this one instead. */
+	FString OverrideSyncStep;
+
+	FSyncSettings(bool bArtist, bool bPreview, FString OverrideSyncStep = FString())
+		: bArtist(bArtist)
+		, bPreview(bPreview)
+		, OverrideSyncStep(MoveTemp(OverrideSyncStep))
+	{ }
+
+	FSyncSettings(const FSyncSettings&& Other)
+		: bArtist(MoveTemp(Other.bArtist))
+		, bPreview(MoveTemp(Other.bPreview))
+		, OverrideSyncStep(MoveTemp(Other.OverrideSyncStep))
+	{ }
+};
+
 /**
  * Helper class with functions used to sync engine.
  */
@@ -250,24 +276,22 @@ public:
 	/**
 	 * Launches UAT UnrealSync command with given command line and options.
 	 *
-	 * @param bArtist Perform artist sync?
-	 * @param bPreview Perform preview sync?
+	 * @param Settings Sync settings.
 	 * @param LabelNameProvider Object that will provide label name to syncing thread.
 	 * @param OnSyncFinished Delegate to run when syncing is finished.
 	 * @param OnSyncProgress Delegate to run when syncing has made progress.
 	 */
-	static void LaunchSync(bool bArtist, bool bPreview, ILabelNameProvider& LabelNameProvider, const FOnSyncFinished& OnSyncFinished, const FOnSyncProgress& OnSyncProgress);
+	static void LaunchSync(FSyncSettings Settings, ILabelNameProvider& LabelNameProvider, const FOnSyncFinished& OnSyncFinished, const FOnSyncProgress& OnSyncProgress);
 
 	/**
 	 * Performs the actual sync with given params.
 	 *
-	 * @param bArtist Perform artist sync?
-	 * @param bPreview Perform preview sync?
+	 * @param Settings Sync settings.
 	 * @param Label Chosen label name.
 	 * @param Game Chosen game.
 	 * @param OnSyncProgress Delegate to run when syncing has made progress. 
 	 */
-	static bool Sync(bool bArtist, bool bPreview, const FString& Label, const FString& Game, const FOnSyncProgress& OnSyncProgress);
+	static bool Sync(const FSyncSettings& Settings, const FString& Label, const FString& Game, const FOnSyncProgress& OnSyncProgress);
 private:
 	/**
 	 * Tries to update original UnrealSync at given location.
