Missing file for previous submit. VS and P4 somehow haven't noticed me that it's not readonly and is not checked out.

diff --git a/Engine/Source/Programs/UnrealSync/Private/UnrealSync.cpp b/Engine/Source/Programs/UnrealSync/Private/UnrealSync.cpp
--- a/Engine/Source/Programs/UnrealSync/Private/UnrealSync.cpp
+++ b/Engine/Source/Programs/UnrealSync/Private/UnrealSync.cpp
@@ -285,14 +285,13 @@ public:
 	/**
 	 * Constructor
 	 *
-	 * @param bArtist Artist sync?
-	 * @param bPreview Preview sync?
+	 * @param Settings Sync settings.
 	 * @param LabelNameProvider Label name provider.
 	 * @param OnSyncFinished Delegate to run when syncing process has finished.
 	 * @param OnSyncProgress Delegate to run when syncing process has made progress.
 	 */
-	FSyncingThread(bool bArtist, bool bPreview, ILabelNameProvider& LabelNameProvider, const FUnrealSync::FOnSyncFinished& OnSyncFinished, const FUnrealSync::FOnSyncProgress& OnSyncProgress)
-		: bArtist(bArtist), bPreview(bPreview), LabelNameProvider(LabelNameProvider), OnSyncFinished(OnSyncFinished), OnSyncProgress(OnSyncProgress)
+	FSyncingThread(FSyncSettings Settings, ILabelNameProvider& LabelNameProvider, const FUnrealSync::FOnSyncFinished& OnSyncFinished, const FUnrealSync::FOnSyncProgress& OnSyncProgress)
+		: Settings(MoveTemp(Settings)), LabelNameProvider(LabelNameProvider), OnSyncFinished(OnSyncFinished), OnSyncProgress(OnSyncProgress)
 	{
 		FRunnableThread::Create(this, TEXT("Syncing thread"));
 	}
@@ -310,7 +309,7 @@ public:
 		FString Label = LabelNameProvider.GetLabelName();
 		FString Game = LabelNameProvider.GetGameName();
 		
-		bool bSuccess = FUnrealSync::Sync(bArtist, bPreview, Label, Game, OnSyncProgress);
+		bool bSuccess = FUnrealSync::Sync(Settings, Label, Game, OnSyncProgress);
 
 		if (OnSyncProgress.IsBound())
 		{
@@ -323,11 +322,8 @@ public:
 	}
 
 private:
-	/* Artist sync? */
-	bool bArtist;
-
-	/* Preview sync? */
-	bool bPreview;
+	/* Sync settings. */
+	FSyncSettings Settings;
 
 	/* Label name provider. */
 	ILabelNameProvider& LabelNameProvider;
@@ -339,9 +335,9 @@ private:
 	FUnrealSync::FOnSyncProgress OnSyncProgress;
 };
 
-void FUnrealSync::LaunchSync(bool bArtist, bool bPreview, ILabelNameProvider& LabelNameProvider, const FOnSyncFinished& OnSyncFinished, const FOnSyncProgress& OnSyncProgress)
+void FUnrealSync::LaunchSync(FSyncSettings Settings, ILabelNameProvider& LabelNameProvider, const FOnSyncFinished& OnSyncFinished, const FOnSyncProgress& OnSyncProgress)
 {
-	SyncingThread = MakeShareable(new FSyncingThread(bArtist, bPreview, LabelNameProvider, OnSyncFinished, OnSyncProgress));
+	SyncingThread = MakeShareable(new FSyncingThread(MoveTemp(Settings), LabelNameProvider, OnSyncFinished, OnSyncProgress));
 }
 
 /**
@@ -428,14 +424,14 @@ static bool GetLatestChangelist(FString& CL)
 	return true;
 }
 
-bool FUnrealSync::Sync(bool bArtist, bool bPreview, const FString& Label, const FString& Game, const FOnSyncProgress& OnSyncProgress)
+bool FUnrealSync::Sync(const FSyncSettings& Settings, const FString& Label, const FString& Game, const FOnSyncProgress& OnSyncProgress)
 {
 	SyncingMessage(OnSyncProgress, "Syncing to label: " + Label);
 
 	auto ProgramRevisionSpec = "@" + Label;
 	TArray<FString> SyncSteps;
 
-	if (bArtist)
+	if (Settings.bArtist)
 	{
 		SyncingMessage(OnSyncProgress, "Performing artist sync.");
 		
@@ -467,7 +463,7 @@ bool FUnrealSync::Sync(bool bArtist, bool bPreview, const FString& Label, const
 	}
 	else
 	{
-		SyncSteps.Add("/..." + ProgramRevisionSpec); // all files to label
+		SyncSteps.Add((Settings.OverrideSyncStep.IsEmpty() ? "/..." : Settings.OverrideSyncStep) + ProgramRevisionSpec); // all files to label
 	}
 
 	for(auto& SyncStep : SyncSteps)
@@ -475,7 +471,7 @@ bool FUnrealSync::Sync(bool bArtist, bool bPreview, const FString& Label, const
 		SyncStep = FP4Env::Get().GetBranch() + SyncStep;
 	}
 
-	if (!FP4Env::RunP4Progress(FString("sync ") + (bPreview ? "-n " : "") + FString::Join(SyncSteps, TEXT(" ")), OnSyncProgress))
+	if (!FP4Env::RunP4Progress(FString("sync ") + (Settings.bPreview ? "-n " : "") + FString::Join(SyncSteps, TEXT(" ")), OnSyncProgress))
 	{
 		return false;
 	}
