UnrealSync: Needed to backout this optimization as P4 admins claim that we cannot trust P4 sync file spec syncing order. We need to figure out some other optimization.

diff --git a/Engine/Source/Programs/UnrealSync/Private/UnrealSync.cpp b/Engine/Source/Programs/UnrealSync/Private/UnrealSync.cpp
--- a/Engine/Source/Programs/UnrealSync/Private/UnrealSync.cpp
+++ b/Engine/Source/Programs/UnrealSync/Private/UnrealSync.cpp
@@ -466,14 +466,13 @@ bool FUnrealSync::Sync(const FSyncSettings& Settings, const FString& Label, cons
 		SyncSteps.Add((Settings.OverrideSyncStep.IsEmpty() ? "/..." : Settings.OverrideSyncStep) + ProgramRevisionSpec); // all files to label
 	}
 
-	for(auto& SyncStep : SyncSteps)
+	FString CommandPrefix = FString("sync ") + (Settings.bPreview ? "-n " : "") + FP4Env::Get().GetBranch();
+	for(const auto& SyncStep : SyncSteps)
 	{
-		SyncStep = FP4Env::Get().GetBranch() + SyncStep;
-	}
-
-	if (!FP4Env::RunP4Progress(FString("sync ") + (Settings.bPreview ? "-n " : "") + FString::Join(SyncSteps, TEXT(" ")), OnSyncProgress))
-	{
-		return false;
+		if (!FP4Env::RunP4Progress(CommandPrefix + SyncStep, OnSyncProgress))
+		{
+			return false;
+		}
 	}
 
 	return true;
