UECORE-120: User is unable to cancel a sync in Unreal Sync

diff --git a/Engine/Source/Programs/UnrealSync/Private/GUI.cpp b/Engine/Source/Programs/UnrealSync/Private/GUI.cpp
--- a/Engine/Source/Programs/UnrealSync/Private/GUI.cpp
+++ b/Engine/Source/Programs/UnrealSync/Private/GUI.cpp
@@ -738,11 +738,13 @@ class SMainTabWidget : public SCompoundWidget
 		EActiveTimerReturnType ExecuteRequests(double CurrentTime, float DeltaTime)
 		{
 			FScopeLock Lock(&RequestsMutex);
-			while (Requests.Num())
+			for (int32 ReqId = 0; ReqId < Requests.Num(); ++ReqId)
 			{
-				Requests.Pop().Execute();
+				Requests[ReqId].Execute();
 			}
 
+			Requests.Empty(4);
+
 			return EActiveTimerReturnType::Continue;
 		}
 
@@ -881,6 +883,14 @@ public:
 					SNew(SHorizontalBox)
 					+ SHorizontalBox::Slot()
 					+ SHorizontalBox::Slot().AutoWidth()
+					[
+						SAssignNew(CancelButton, SButton)
+						.OnClicked(this, &SMainTabWidget::OnCancelButtonClick)
+						[
+							SNew(STextBlock).Text(FText::FromString("Cancel"))
+						]
+					]
+					+ SHorizontalBox::Slot().AutoWidth()
 					[
 						GoBackButton.ToSharedRef()
 					]
@@ -923,13 +933,14 @@ public:
 	virtual ~SMainTabWidget()
 	{
 		FUnrealSync::TerminateLoadingProcess();
+		FUnrealSync::TerminateSyncingProcess();
 	}
 
 private:
 	/**
 	 * Queues adding lines to the log for execution on Slate rendering thread.
 	 */
-	void ExternalThreadAddLinesToLog(TArray<TSharedPtr<FString> > Lines)
+	void ExternalThreadAddLinesToLog(TSharedPtr<TArray<TSharedPtr<FString> > > Lines)
 	{
 		ExternalThreadsDispatcher->AddRenderingThreadRequest(FExternalThreadsDispatcher::FRenderingThreadRequest::CreateSP(this, &SMainTabWidget::AddLinesToLog, Lines));
 	}
@@ -939,9 +950,9 @@ private:
 	 *
 	 * @param Lines Lines to add.
 	 */
-	void AddLinesToLog(TArray<TSharedPtr<FString> > Lines)
+	void AddLinesToLog(TSharedPtr<TArray<TSharedPtr<FString> > > Lines)
 	{
-		for (const auto& Line : Lines)
+		for (const auto& Line : *Lines)
 		{
 			LogLines.Add(Line);
 		}
@@ -986,17 +997,42 @@ private:
 	}
 
 	/**
-	 * Function that is called on when go back button is clicked.
+	 * Performs a procedure that needs to be done when GoBack button is clicked.
 	 * It tells the app to go back to the beginning.
-	 *
-	 * @returns Tells that this event was handled.
 	 */
-	FReply OnGoBackButtonClick()
+	void GoBack()
 	{
 		Switcher->SetActiveWidgetIndex(0);
 		LogLines.Empty();
 		LogListView->RequestListRefresh();
 		OnReloadLabels();
+	}
+
+	/**
+	 * Function that is called on when go back button is clicked.
+	 * It tells the app to go back to the beginning.
+	 *
+	 * @returns Tells that this event was handled.
+	 */
+	FReply OnGoBackButtonClick()
+	{
+		GoBack();
+
+		return FReply::Handled();
+	}
+
+	/**
+	 * Function that is called on when cancel button is clicked.
+	 * It tells the app to cancel the sync and go back to the beginning.
+	 *
+	 * @returns Tells that this event was handled.
+	 */
+	FReply OnCancelButtonClick()
+	{
+		FUnrealSync::TerminateLoadingProcess();
+		FUnrealSync::TerminateSyncingProcess();
+
+		ExternalThreadsDispatcher->AddRenderingThreadRequest(FExternalThreadsDispatcher::FRenderingThreadRequest::CreateSP(this, &SMainTabWidget::GoBack));
 
 		return FReply::Handled();
 	}
@@ -1053,6 +1089,7 @@ private:
 			FUnrealSync::FOnSyncProgress::CreateRaw(this, &SMainTabWidget::SyncingProgress));
 
 		GoBackButton->SetEnabled(false);
+		CancelButton->SetEnabled(true);
 
 		return FReply::Handled();
 	}
@@ -1078,11 +1115,11 @@ private:
 		FString Line;
 		FString Rest;
 
-		TArray<TSharedPtr<FString> > Lines;
+		TSharedPtr<TArray<TSharedPtr<FString> > > Lines = MakeShareable(new TArray<TSharedPtr<FString> >());
 
 		while (Buffer.Split("\n", &Line, &Rest))
 		{
-			Lines.Add(MakeShareable(new FString(Line)));
+			Lines->Add(MakeShareable(new FString(Line)));
 
 			Buffer = Rest;
 		}
@@ -1099,6 +1136,7 @@ private:
 	 */
 	void SyncingFinished(bool bSuccess)
 	{
+		CancelButton->SetEnabled(false);
 		GoBackButton->SetEnabled(true);
 	}
 
@@ -1195,6 +1233,8 @@ private:
 
 	/* Go back button reference. */
 	TSharedPtr<SButton> GoBackButton;
+	/* Cancel button reference. */
+	TSharedPtr<SButton> CancelButton;
 
 	/* Array of sync command line providers. */
 	TArray<TSharedRef<ILabelNameProvider> > SyncCommandLineProviders;
diff --git a/Engine/Source/Programs/UnrealSync/Private/UnrealSync.cpp b/Engine/Source/Programs/UnrealSync/Private/UnrealSync.cpp
--- a/Engine/Source/Programs/UnrealSync/Private/UnrealSync.cpp
+++ b/Engine/Source/Programs/UnrealSync/Private/UnrealSync.cpp
@@ -291,9 +291,9 @@ public:
 	 * @param OnSyncProgress Delegate to run when syncing process has made progress.
 	 */
 	FSyncingThread(FSyncSettings Settings, ILabelNameProvider& LabelNameProvider, const FUnrealSync::FOnSyncFinished& OnSyncFinished, const FUnrealSync::FOnSyncProgress& OnSyncProgress)
-		: Settings(MoveTemp(Settings)), LabelNameProvider(LabelNameProvider), OnSyncFinished(OnSyncFinished), OnSyncProgress(OnSyncProgress)
+		: Settings(MoveTemp(Settings)), LabelNameProvider(LabelNameProvider), OnSyncFinished(OnSyncFinished), OnSyncProgress(OnSyncProgress), bTerminate(false)
 	{
-		FRunnableThread::Create(this, TEXT("Syncing thread"));
+		Thread = FRunnableThread::Create(this, TEXT("Syncing thread"));
 	}
 
 	/**
@@ -308,8 +308,32 @@ public:
 
 		FString Label = LabelNameProvider.GetLabelName();
 		FString Game = LabelNameProvider.GetGameName();
+
+		struct FProcessStopper
+		{
+			FProcessStopper(bool& bStop, FUnrealSync::FOnSyncProgress& OuterSyncProgress)
+				: bStop(bStop), OuterSyncProgress(OuterSyncProgress) {}
+
+			bool OnProgress(const FString& Text)
+			{
+				if (OuterSyncProgress.IsBound())
+				{
+					if (!OuterSyncProgress.Execute(Text))
+					{
+						bStop = true;
+					}
+				}
+
+				return !bStop;
+			}
+
+		private:
+			bool& bStop;
+			FUnrealSync::FOnSyncProgress& OuterSyncProgress;
+		};
 		
-		bool bSuccess = FUnrealSync::Sync(Settings, Label, Game, OnSyncProgress);
+		FProcessStopper Stopper(bTerminate, OnSyncProgress);
+		bool bSuccess = FUnrealSync::Sync(Settings, Label, Game, FUnrealSync::FOnSyncProgress::CreateRaw(&Stopper, &FProcessStopper::OnProgress));
 
 		if (OnSyncProgress.IsBound())
 		{
@@ -321,7 +345,23 @@ public:
 		return 0;
 	}
 
+	/**
+	 * Stops process runnning in the background and terminates wait for the
+	 * watcher thread to finish.
+	 */
+	void Terminate()
+	{
+		bTerminate = true;
+		Thread->WaitForCompletion();
+	}
+
 private:
+	/* Tells the thread to terminate the process. */
+	bool bTerminate;
+
+	/* Handle for thread object. */
+	FRunnableThread* Thread;
+
 	/* Sync settings. */
 	FSyncSettings Settings;
 
@@ -335,6 +375,14 @@ private:
 	FUnrealSync::FOnSyncProgress OnSyncProgress;
 };
 
+void FUnrealSync::TerminateSyncingProcess()
+{
+	if (SyncingThread.IsValid())
+	{
+		SyncingThread->Terminate();
+	}
+}
+
 void FUnrealSync::LaunchSync(FSyncSettings Settings, ILabelNameProvider& LabelNameProvider, const FOnSyncFinished& OnSyncFinished, const FOnSyncProgress& OnSyncProgress)
 {
 	SyncingThread = MakeShareable(new FSyncingThread(MoveTemp(Settings), LabelNameProvider, OnSyncFinished, OnSyncProgress));
diff --git a/Engine/Source/Programs/UnrealSync/Private/UnrealSync.h b/Engine/Source/Programs/UnrealSync/Private/UnrealSync.h
--- a/Engine/Source/Programs/UnrealSync/Private/UnrealSync.h
+++ b/Engine/Source/Programs/UnrealSync/Private/UnrealSync.h
@@ -238,6 +238,11 @@ public:
 	 */
 	static void TerminateLoadingProcess();
 
+	/**
+	 * Terminates P4 syncing process.
+	 */
+	static void TerminateSyncingProcess();
+
 	/**
 	 * Method to receive p4 data loading finished event.
 	 *
