UECORE-118: Unable to Launch Editor from Unreal Sync

diff --git a/Engine/Source/Programs/UnrealSync/Private/GUI.cpp b/Engine/Source/Programs/UnrealSync/Private/GUI.cpp
--- a/Engine/Source/Programs/UnrealSync/Private/GUI.cpp
+++ b/Engine/Source/Programs/UnrealSync/Private/GUI.cpp
@@ -859,7 +859,7 @@ public:
 								PreviewSyncCheckBox.ToSharedRef()
 							]
 					]
-					+ SHorizontalBox::Slot().AutoWidth().Padding(FMargin(30, 0, 0, 0))
+					+ SHorizontalBox::Slot().AutoWidth().Padding(30, 0, 0, 0)
 						[
 							SNew(SButton).HAlign(HAlign_Right)
 							.OnClicked(this, &SMainTabWidget::OnSync)
@@ -869,6 +869,15 @@ public:
 								SNew(STextBlock).Text(FText::FromString("Sync!"))
 							]
 						]
+					+ SHorizontalBox::Slot().AutoWidth()
+						[
+							SNew(SButton).HAlign(HAlign_Right)
+							.OnClicked(this, &SMainTabWidget::RunUE4)
+							.VAlign(VAlign_Center)
+							[
+								SNew(STextBlock).Text(FText::FromString("Run UE4"))
+							]
+						]
 				]
 			]
 			+ SWidgetSwitcher::Slot()
@@ -1070,6 +1079,22 @@ private:
 		return FReply::Handled();
 	}
 
+	/**
+	 * Executes UE4Editor.exe for current branch.
+	 */
+	FReply RunUE4()
+	{
+#if PLATFORM_WINDOWS
+		auto ProcessPath = FPaths::Combine(*FPaths::GetPath(FPlatformProcess::GetApplicationName(FPlatformProcess::GetCurrentProcessId())), TEXT("UE4Editor.exe"));
+#else
+		// Needs to be implemented on other platforms.
+#endif
+
+		RunProcess(ProcessPath);
+
+		return FReply::Handled();
+	}
+
 	/**
 	 * Sync button click function. Gets all widget data, constructs command
 	 * line and launches syncing.
diff --git a/Engine/Source/Programs/UnrealSync/Private/ProcessHelper.h b/Engine/Source/Programs/UnrealSync/Private/ProcessHelper.h
--- a/Engine/Source/Programs/UnrealSync/Private/ProcessHelper.h
+++ b/Engine/Source/Programs/UnrealSync/Private/ProcessHelper.h
@@ -12,7 +12,7 @@ DECLARE_DELEGATE_RetVal_OneParam(bool, FOnProcessMadeProgress, const FString&)
  *
  * @returns True if succeeded. False otherwise.
  */
-bool RunProcess(const FString& ExecutablePath, const FString& CommandLine);
+bool RunProcess(const FString& ExecutablePath, const FString& CommandLine = FString());
 
 /**
  * Helper function to run custom command line and catch output.
