UECORE-116: Add Progress Bar / Animation to US3

diff --git a/Engine/Source/Programs/UnrealSync/Private/GUI.cpp b/Engine/Source/Programs/UnrealSync/Private/GUI.cpp
--- a/Engine/Source/Programs/UnrealSync/Private/GUI.cpp
+++ b/Engine/Source/Programs/UnrealSync/Private/GUI.cpp
@@ -12,6 +12,7 @@
 #include "SWidgetSwitcher.h"
 #include "SDockTab.h"
 #include "SlateFwd.h"
+#include "SThrobber.h"
 
 /**
  * Simple text combo box widget.
@@ -891,6 +892,10 @@ public:
 				[
 					SNew(SHorizontalBox)
 					+ SHorizontalBox::Slot()
+					[
+						SAssignNew(Throbber, SThrobber)
+						.Animate(SThrobber::VerticalAndOpacity).NumPieces(5)
+					]
 					+ SHorizontalBox::Slot().AutoWidth()
 					[
 						SAssignNew(CancelButton, SButton)
@@ -1115,6 +1120,7 @@ private:
 
 		GoBackButton->SetEnabled(false);
 		CancelButton->SetEnabled(true);
+		Throbber->SetAnimate(SThrobber::EAnimation::VerticalAndOpacity);
 
 		return FReply::Handled();
 	}
@@ -1161,6 +1167,7 @@ private:
 	 */
 	void SyncingFinished(bool bSuccess)
 	{
+		Throbber->SetAnimate(SThrobber::EAnimation::None);
 		CancelButton->SetEnabled(false);
 		GoBackButton->SetEnabled(true);
 	}
@@ -1260,6 +1267,8 @@ private:
 	TSharedPtr<SButton> GoBackButton;
 	/* Cancel button reference. */
 	TSharedPtr<SButton> CancelButton;
+	/* Throbber to indicate work-in-progress. */
+	TSharedPtr<SThrobber> Throbber;
 
 	/* Array of sync command line providers. */
 	TArray<TSharedRef<ILabelNameProvider> > SyncCommandLineProviders;
