UECORE-119: Add prompt to close UE4Editor when trying to sync with Unreal Sync.

diff --git a/Engine/Source/Programs/UnrealSync/Private/GUI.cpp b/Engine/Source/Programs/UnrealSync/Private/GUI.cpp
--- a/Engine/Source/Programs/UnrealSync/Private/GUI.cpp
+++ b/Engine/Source/Programs/UnrealSync/Private/GUI.cpp
@@ -14,6 +14,159 @@
 #include "SlateFwd.h"
 #include "SThrobber.h"
 
+/**
+ * Performs a dialog when running editor is detected.
+ */
+class SEditorRunningDlg : public SCompoundWidget
+{
+public:
+	enum class EResponse
+	{
+		ForceSync,
+		Retry,
+		Cancel
+	};
+
+	SLATE_BEGIN_ARGS(SEditorRunningDlg)	{}
+		SLATE_ATTRIBUTE(TSharedPtr<SWindow>, ParentWindow)
+		SLATE_ATTRIBUTE(FText, Message)
+		SLATE_ATTRIBUTE(float, WrapMessageAt)
+	SLATE_END_ARGS()
+
+	BEGIN_SLATE_FUNCTION_BUILD_OPTIMIZATION
+	void Construct(const FArguments& InArgs)
+	{
+		ParentWindow = InArgs._ParentWindow.Get();
+		ParentWindow->SetWidgetToFocusOnActivate(SharedThis(this));
+
+		Response = EResponse::Retry;
+
+		FSlateFontInfo MessageFont(FCoreStyle::Get().GetFontStyle("StandardDialog.LargeFont"));
+
+		this->ChildSlot
+			[
+				SNew(SBorder)
+				.BorderImage(FCoreStyle::Get().GetBrush("ToolPanel.GroupBorder"))
+				[
+					SNew(SVerticalBox)
+
+					+ SVerticalBox::Slot()
+					.HAlign(HAlign_Fill)
+					.VAlign(VAlign_Fill)
+					.FillHeight(1.0f)
+					.MaxHeight(550)
+					.Padding(12.0f)
+					[
+						SNew(SScrollBox)
+
+						+ SScrollBox::Slot()
+						[
+							SNew(STextBlock)
+							.Text(LOCTEXT("UE4EditorInstanceDetectedDialog", "UE4Editor detected currently running on the system. Please close it before proceeding."))
+							.Font(MessageFont)
+							.WrapTextAt(InArgs._WrapMessageAt)
+						]
+					]
+
+					+ SVerticalBox::Slot()
+						.AutoHeight()
+						.Padding(0.0f)
+						[
+							SNew(SHorizontalBox)
+							+ SHorizontalBox::Slot()
+							+ SHorizontalBox::Slot()
+								.AutoWidth()
+								.HAlign(HAlign_Right)
+								.VAlign(VAlign_Bottom)
+								.Padding(2.f)
+								[
+									SNew(SUniformGridPanel)
+									.SlotPadding(FCoreStyle::Get().GetMargin("StandardDialog.SlotPadding"))
+									.MinDesiredSlotWidth(FCoreStyle::Get().GetFloat("StandardDialog.MinDesiredSlotWidth"))
+									.MinDesiredSlotHeight(FCoreStyle::Get().GetFloat("StandardDialog.MinDesiredSlotHeight"))
+									+ SUniformGridPanel::Slot(0, 0)
+									[
+										SNew(SButton)
+										.Text(LOCTEXT("CancelSync", "Cancel sync"))
+										.OnClicked(this, &SEditorRunningDlg::HandleButtonClicked, EResponse::Cancel)
+										.ContentPadding(FCoreStyle::Get().GetMargin("StandardDialog.ContentPadding"))
+										.HAlign(HAlign_Center)
+									]
+									+ SUniformGridPanel::Slot(1, 0)
+									[
+										SNew(SButton)
+										.Text(LOCTEXT("RetryEditorRunning", "Retry if editor is running"))
+										.OnClicked(this, &SEditorRunningDlg::HandleButtonClicked, EResponse::Retry)
+										.ContentPadding(FCoreStyle::Get().GetMargin("StandardDialog.ContentPadding"))
+										.HAlign(HAlign_Center)
+									]
+									+ SUniformGridPanel::Slot(2, 0)
+									[
+										SNew(SButton)
+										.Text(LOCTEXT("ForceSync", "Force sync"))
+										.OnClicked(this, &SEditorRunningDlg::HandleButtonClicked, EResponse::ForceSync)
+										.ContentPadding(FCoreStyle::Get().GetMargin("StandardDialog.ContentPadding"))
+										.HAlign(HAlign_Center)
+									]
+								]
+						]
+				]
+			];
+	}
+	END_SLATE_FUNCTION_BUILD_OPTIMIZATION
+
+	/**
+	 * Shows modal dialog window and gets it's response.
+	 */
+	static EResponse ShowAndGetResponse()
+	{
+		TSharedPtr<SWindow> OutWindow = SNew(SWindow)
+			.Title(LOCTEXT("UE4EditorInstanceDetected", "UE4Editor instance detected!"))
+			.SizingRule(ESizingRule::Autosized)
+			.AutoCenter(EAutoCenter::PreferredWorkArea)
+			.SupportsMinimize(false).SupportsMaximize(false);
+
+		TSharedPtr<SEditorRunningDlg> OutDialog = SNew(SEditorRunningDlg)
+			.ParentWindow(OutWindow)
+			.WrapMessageAt(512.0f);
+
+		OutWindow->SetContent(OutDialog.ToSharedRef());
+
+		FSlateApplication::Get().AddModalWindow(OutWindow.ToSharedRef(), FGlobalTabmanager::Get()->GetRootWindow());
+
+		return OutDialog->GetResponse();
+	}
+
+private:
+	/**
+	 * Gets response from this dialog widget.
+	 *
+	 * @returns Response.
+	 */
+	EResponse GetResponse() const
+	{
+		return Response;
+	}
+
+	/**
+	 * Handles button click on any dialog button.
+	 */
+	FReply HandleButtonClicked(EResponse Response)
+	{
+		this->Response = Response;
+
+		ParentWindow->RequestDestroyWindow();
+
+		return FReply::Handled();
+	}
+
+	/** Parent window pointer */
+	TSharedPtr<SWindow> ParentWindow;
+
+	/** Saved response */
+	EResponse Response;
+};
+
 /**
  * Simple text combo box widget.
  */
@@ -1108,6 +1261,11 @@ private:
 	 */
 	FReply OnSync()
 	{
+		if (CheckIfEditorIsRunning())
+		{
+			return FReply::Handled();
+		}
+
 		FSyncSettings Settings(
 			ArtistSyncCheckBox->IsChecked(),
 			PreviewSyncCheckBox->IsChecked(),
@@ -1242,6 +1400,64 @@ private:
 		SyncCommandLineProviders.Add(Widget);
 	}
 
+	/**
+	 * Checks if the editor is currently running on the system and asks user to
+	 * close it.
+	 *
+	 * @returns False if editor is not opened or user forced to continue.
+	 *          Otherwise true.
+	 */
+	bool CheckIfEditorIsRunning()
+	{
+		auto Response = SEditorRunningDlg::EResponse::Retry;
+		bool bEditorIsRunning = false;
+			 
+		while (Response == SEditorRunningDlg::EResponse::Retry && ((bEditorIsRunning = CheckForEditorProcess()) == true))
+		{
+			Response = SEditorRunningDlg::ShowAndGetResponse();
+		}
+
+		return bEditorIsRunning && Response != SEditorRunningDlg::EResponse::ForceSync;
+	}
+
+	/**
+	 * Checks if the editor is currently running on the system.
+	 *
+	 * @returns False if editor is not opened. Otherwise true.
+	 */
+	bool CheckForEditorProcess()
+	{
+		for (const auto& PossibleExecutablePath : GetPossibleExecutablePaths())
+		{
+			if (IsRunningProcess(PossibleExecutablePath))
+			{
+				return true;
+			}
+		}
+
+		return false;
+	}
+
+#if PLATFORM_WINDOWS
+	/**
+	 * Returns possible colliding UE4Editor executables.
+	 *
+	 * @returns Array of possible executables.
+	 */
+	static TArray<FString> GetPossibleExecutablePaths()
+	{
+		FString BasePath = FPaths::GetPath(FPlatformProcess::GetApplicationName(FPlatformProcess::GetCurrentProcessId()));
+
+		TArray<FString> Out;
+
+		Out.Add(FPaths::Combine(*BasePath, TEXT("UE4Editor.exe")));
+		Out.Add(FPaths::Combine(*BasePath, TEXT("UE4Editor-Win32-Debug.exe")));
+		Out.Add(FPaths::Combine(*BasePath, TEXT("UE4Editor-Win64-Debug.exe")));
+
+		return Out;
+	}
+#endif
+
 	/* Main widget switcher. */
 	TSharedPtr<SWidgetSwitcher> Switcher;
 
diff --git a/Engine/Source/Programs/UnrealSync/Private/ProcessHelper.cpp b/Engine/Source/Programs/UnrealSync/Private/ProcessHelper.cpp
--- a/Engine/Source/Programs/UnrealSync/Private/ProcessHelper.cpp
+++ b/Engine/Source/Programs/UnrealSync/Private/ProcessHelper.cpp
@@ -91,4 +91,42 @@ bool RunProcessProgress(const FString& ExecutablePath, const FString& CommandLin
 	}
 
 	return ReturnCode == 0;
-}
\ No newline at end of file
+}
+
+#if PLATFORM_WINDOWS
+
+#define DWORD int32
+
+#include <tlhelp32.h>
+
+bool IsRunningProcess(const FString& FullImagePath)
+{
+	auto InFullImagePath = FullImagePath.ToUpper();
+	FPaths::NormalizeFilename(InFullImagePath);
+
+	// append the extension
+	HANDLE SnapShot = ::CreateToolhelp32Snapshot(TH32CS_SNAPPROCESS, 0);
+	if (SnapShot != INVALID_HANDLE_VALUE)
+	{
+		PROCESSENTRY32 Entry;
+		Entry.dwSize = sizeof(PROCESSENTRY32);
+
+		if (::Process32First(SnapShot, &Entry))
+		{
+			do
+			{
+				auto AppName = FPlatformProcess::GetApplicationName(Entry.th32ProcessID).ToUpper();
+				FPaths::NormalizeFilename(AppName);
+				if (AppName.Compare(InFullImagePath) == 0)
+				{
+					::CloseHandle(SnapShot);
+					return true;
+				}
+			} while (::Process32Next(SnapShot, &Entry));
+		}
+	}
+
+	::CloseHandle(SnapShot);
+	return false;
+}
+#endif
diff --git a/Engine/Source/Programs/UnrealSync/Private/ProcessHelper.h b/Engine/Source/Programs/UnrealSync/Private/ProcessHelper.h
--- a/Engine/Source/Programs/UnrealSync/Private/ProcessHelper.h
+++ b/Engine/Source/Programs/UnrealSync/Private/ProcessHelper.h
@@ -34,4 +34,15 @@ bool RunProcessOutput(const FString& ExecutablePath, const FString& CommandLine,
  *
  * @returns True if succeeded. False otherwise.
  */
-bool RunProcessProgress(const FString& ExecutablePath, const FString& CommandLine, const FOnProcessMadeProgress& OnUATMadeProgress);
\ No newline at end of file
+bool RunProcessProgress(const FString& ExecutablePath, const FString& CommandLine, const FOnProcessMadeProgress& OnUATMadeProgress);
+
+#if PLATFORM_WINDOWS
+/**
+ * Checks if process is currently running given full image path.
+ *
+ * @param FullImagePath Full path to the process executable.
+ *
+ * @returns True if it's running. False otherwise.
+ */
+bool IsRunningProcess(const FString& FullImagePath);
+#endif
\ No newline at end of file
