UECORE-115: Add copy/save/clear functionality for Unreal Sync logging

diff --git a/Engine/Source/Programs/UnrealSync/Private/GUI.cpp b/Engine/Source/Programs/UnrealSync/Private/GUI.cpp
--- a/Engine/Source/Programs/UnrealSync/Private/GUI.cpp
+++ b/Engine/Source/Programs/UnrealSync/Private/GUI.cpp
@@ -13,6 +13,114 @@
 #include "SDockTab.h"
 #include "SlateFwd.h"
 #include "SThrobber.h"
+#include "SMultiLineEditableTextBox.h"
+#include "BaseTextLayoutMarshaller.h"
+
+/**
+ * Syncing log.
+ */
+class SSyncLog : public SCompoundWidget
+{
+public:
+	SLATE_BEGIN_ARGS(SSyncLog) {}
+	SLATE_END_ARGS()
+
+	BEGIN_SLATE_FUNCTION_BUILD_OPTIMIZATION
+	void Construct(const FArguments& InArgs)
+	{
+		MessagesTextMarshaller = MakeShareable(new FMarshaller());
+		MessagesTextBox = SNew(SMultiLineEditableTextBox)
+			.IsReadOnly(true)
+			.AlwaysShowScrollbars(true)
+			.Marshaller(MessagesTextMarshaller);
+
+		ChildSlot
+			[
+				MessagesTextBox.ToSharedRef()
+			];
+	}
+	END_SLATE_FUNCTION_BUILD_OPTIMIZATION
+
+	/**
+	 * Adding messages to the log.
+	 *
+	 * @param Message Message to add.
+	 */
+	void AddMessage(const FString& Message)
+	{
+		MessagesTextMarshaller->AddMessage(Message);
+		++NumberOfLines;
+		MessagesTextBox->ScrollTo(FTextLocation(NumberOfLines));
+	}
+
+	/**
+	 * Clear the log.
+	 */	
+	void Clear()
+	{
+		NumberOfLines = 0;
+		MessagesTextBox->GoTo(FTextLocation(0));
+
+		MessagesTextMarshaller->Clear();
+		MessagesTextBox->Refresh();
+	}
+
+private:
+	/**
+	 * Messages marshaller.
+	 */
+	class FMarshaller : public FBaseTextLayoutMarshaller
+	{
+	public:
+		void AddMessage(const FString& Message)
+		{
+			auto NormalText = FTextBlockStyle()
+				.SetFont("Fonts/Roboto-Regular", 10)
+				.SetColorAndOpacity(FSlateColor::UseForeground())
+				.SetShadowOffset(FVector2D::ZeroVector)
+				.SetShadowColorAndOpacity(FLinearColor::Black);
+
+			TSharedRef<FString> LineText = MakeShareable(new FString(Message));
+
+			TArray<TSharedRef<IRun>> Runs;
+			Runs.Add(FSlateTextRun::Create(FRunInfo(), LineText, NormalText));
+
+			Layout->AddLine(LineText, Runs);
+		}
+
+		/**
+		 * Clears messages from the text box.
+		 */
+		void Clear()
+		{
+			MakeDirty();
+		}
+
+	private:
+		// FBaseTextLayoutMarshaller missing virtuals.
+		virtual void SetText(const FString& SourceString, FTextLayout& TargetTextLayout)
+		{
+			// Get text layout.
+			Layout = &TargetTextLayout;
+		}
+
+		virtual void GetText(FString& TargetString, const FTextLayout& SourceTextLayout)
+		{
+			SourceTextLayout.GetAsText(TargetString);
+		}
+
+		FTextLayout* Layout;
+	};
+
+	/** Converts the array of messages into something the text box understands */
+	TSharedPtr<FMarshaller> MessagesTextMarshaller;
+
+	/** The editable text showing all log messages */
+	TSharedPtr<SMultiLineEditableTextBox> MessagesTextBox;
+
+	/** Number of currently logged lines. */
+	int32 NumberOfLines = 0;
+};
 
 /**
  * Performs a dialog when running editor is detected.
@@ -969,10 +1077,6 @@ public:
 				SNew(STextBlock).Text(FText::FromString("Go back"))
 			];
 
-		LogListView = SNew(SListView<TSharedPtr<FString> >)
-			.ListItemsSource(&LogLines)
-			.OnGenerateRow(this, &SMainTabWidget::GenerateLogItem);
-
 		TSharedPtr<SVerticalBox> MainBox;
 		
 		Switcher = SNew(SWidgetSwitcher)
@@ -1039,7 +1143,7 @@ public:
 				SNew(SVerticalBox)
 				+ SVerticalBox::Slot()
 				[
-					LogListView->AsWidget()
+					SAssignNew(Log, SSyncLog)
 				]
 				+ SVerticalBox::Slot().AutoHeight()
 				[
@@ -1109,7 +1213,7 @@ private:
 	 */
 	void ExternalThreadAddLinesToLog(TSharedPtr<TArray<TSharedPtr<FString> > > Lines)
 	{
-		ExternalThreadsDispatcher->AddRenderingThreadRequest(FExternalThreadsDispatcher::FRenderingThreadRequest::CreateSP(this, &SMainTabWidget::AddLinesToLog, Lines));
+		ExternalThreadsDispatcher->AddRenderingThreadRequest(FExternalThreadsDispatcher::FRenderingThreadRequest::CreateRaw(this, &SMainTabWidget::AddLinesToLog, Lines));
 	}
 
 	/**
@@ -1121,10 +1225,8 @@ private:
 	{
 		for (const auto& Line : *Lines)
 		{
-			LogLines.Add(Line);
+			Log->AddMessage(*Line.Get());
 		}
-
-		LogListView->RequestScrollIntoView(LogLines.Last());
 	}
 
 	/**
@@ -1170,8 +1272,7 @@ private:
 	void GoBack()
 	{
 		Switcher->SetActiveWidgetIndex(0);
-		LogLines.Empty();
-		LogListView->RequestListRefresh();
+		Log->Clear();
 		OnReloadLabels();
 	}
 
@@ -1474,10 +1575,8 @@ private:
 	/* External thread requests dispatcher. */
 	TSharedRef<FExternalThreadsDispatcher, ESPMode::ThreadSafe> ExternalThreadsDispatcher;
 
-	/* Report log. */
-	TArray<TSharedPtr<FString> > LogLines;
-	/* Log list view. */
-	TSharedPtr<SListView<TSharedPtr<FString> > > LogListView;
+	/* Syncing log. */
+	TSharedPtr<SSyncLog> Log;
 
 	/* Go back button reference. */
 	TSharedPtr<SButton> GoBackButton;
diff --git a/Engine/Source/Programs/UnrealSync/UnrealSync.Build.cs b/Engine/Source/Programs/UnrealSync/UnrealSync.Build.cs
--- a/Engine/Source/Programs/UnrealSync/UnrealSync.Build.cs
+++ b/Engine/Source/Programs/UnrealSync/UnrealSync.Build.cs
@@ -16,7 +16,6 @@ public class UnrealSync : ModuleRules
 				"Projects",
 				"Slate",
 				"SlateCore",
-				"SlateReflector",
 				"StandaloneRenderer",
 				"XmlParser",
 				"JSON"
