UE-14033: Selecting Go Back closes Unreal Sync

diff --git a/Engine/Source/Programs/UnrealSync/Private/GUI.cpp b/Engine/Source/Programs/UnrealSync/Private/GUI.cpp
--- a/Engine/Source/Programs/UnrealSync/Private/GUI.cpp
+++ b/Engine/Source/Programs/UnrealSync/Private/GUI.cpp
@@ -34,6 +34,8 @@ public:
 			.AlwaysShowScrollbars(true)
 			.Marshaller(MessagesTextMarshaller);
 
+		Clear();
+
 		ChildSlot
 			[
 				MessagesTextBox.ToSharedRef()
@@ -49,8 +51,7 @@ public:
 	void AddMessage(const FString& Message)
 	{
 		MessagesTextMarshaller->AddMessage(Message);
-		++NumberOfLines;
-		MessagesTextBox->ScrollTo(FTextLocation(NumberOfLines));
+		MessagesTextBox->ScrollTo(FTextLocation(NumberOfLines++));
 	}
 
 	/**
@@ -59,8 +60,8 @@ public:
 	void Clear()
 	{
 		NumberOfLines = 0;
-		MessagesTextBox->GoTo(FTextLocation(0));
-
+		MessagesTextBox->ScrollTo(FTextLocation(0));
+		
 		MessagesTextMarshaller->Clear();
 		MessagesTextBox->Refresh();
 	}
@@ -85,6 +86,11 @@ private:
 			TArray<TSharedRef<IRun>> Runs;
 			Runs.Add(FSlateTextRun::Create(FRunInfo(), LineText, NormalText));
 
+			if (Layout->IsEmpty())
+			{
+				Layout->ClearLines();
+			}
+
 			Layout->AddLine(LineText, Runs);
 		}
 
@@ -1300,7 +1306,7 @@ private:
 		FUnrealSync::TerminateLoadingProcess();
 		FUnrealSync::TerminateSyncingProcess();
 
-		ExternalThreadsDispatcher->AddRenderingThreadRequest(FExternalThreadsDispatcher::FRenderingThreadRequest::CreateSP(this, &SMainTabWidget::GoBack));
+		ExternalThreadsDispatcher->AddRenderingThreadRequest(FExternalThreadsDispatcher::FRenderingThreadRequest::CreateRaw(this, &SMainTabWidget::GoBack));
 
 		return FReply::Handled();
 	}
@@ -1373,7 +1379,7 @@ private:
 			OverrideSyncStep);
 
 		Switcher->SetActiveWidgetIndex(1);
-		FUnrealSync::LaunchSync(Settings, GetCurrentSyncCmdLineProvider(),
+		FUnrealSync::LaunchSync(MoveTemp(Settings), GetCurrentSyncCmdLineProvider(),
 			FUnrealSync::FOnSyncFinished::CreateRaw(this, &SMainTabWidget::SyncingFinished),
 			FUnrealSync::FOnSyncProgress::CreateRaw(this, &SMainTabWidget::SyncingProgress));
 
