UE-1542: Unable to set Auto Clobber for Unreal Sync

diff --git a/Engine/Source/Programs/UnrealSync/Private/GUI.cpp b/Engine/Source/Programs/UnrealSync/Private/GUI.cpp
--- a/Engine/Source/Programs/UnrealSync/Private/GUI.cpp
+++ b/Engine/Source/Programs/UnrealSync/Private/GUI.cpp
@@ -319,6 +319,7 @@ public:
 	 */
 	void Add(const FString& Value)
 	{
+		CurrentOption = 0;
 		OptionsSource.Add(MakeShareable(new FString(Value)));
 	}
 
@@ -347,6 +348,7 @@ public:
 	 */
 	void Clear()
 	{
+		CurrentOption = -1;
 		OptionsSource.Empty();
 	}
 
@@ -423,7 +425,7 @@ public:
 		 * @param Name Name of the item.
 		 * @param Content Content widget to store by this item.
 		 */
-		FItem(SRadioContentSelection& Parent, int32 Id, const FString& Name, TSharedRef<SWidget> Content)
+		FItem(SRadioContentSelection& Parent, int32 Id, FText Name, TSharedRef<SWidget> Content)
 			: Parent(Parent), Id(Id), Name(Name), Content(Content)
 		{
 
@@ -445,7 +447,7 @@ public:
 					.IsChecked(this, &FItem::IsCheckboxChecked)
 					[
 						SNew(STextBlock)
-						.Text(FText::FromString(GetName()))
+						.Text(Name)
 					]
 				]
 				+ SVerticalBox::Slot().VAlign(VAlign_Fill).Padding(19.0f, 2.0f, 2.0f, 2.0f)
@@ -454,16 +456,6 @@ public:
 				];
 		}
 
-		/**
-		 * Gets this item name.
-		 *
-		 * @returns Item name.
-		 */
-		const FString& GetName() const
-		{
-			return Name;
-		}
-		
 		/**
 		 * Gets this item content widget.
 		 *
@@ -526,7 +518,7 @@ public:
 		/* Item id. */
 		int32 Id;
 		/* Item name. */
-		FString Name;
+		FText Name;
 
 		/* Outer item widget. */
 		TSharedPtr<SWidget> WholeItemWidget;
@@ -561,7 +553,7 @@ public:
 	 * @param Name Name of the widget in this selection.
 	 * @param Content The widget itself.
 	 */
-	void Add(const FString& Name, TSharedRef<SWidget> Content)
+	void Add(FText Name, TSharedRef<SWidget> Content)
 	{
 		TSharedRef<FItem> Item = MakeShareable(new FItem(*this, Items.Num(), Name, Content));
 		Item->Init();
@@ -1046,17 +1038,17 @@ public:
 	{
 		RadioSelection = SNew(SRadioContentSelection);
 
-		AddToRadioSelection("Sync to the latest promoted", SNew(SLatestPromoted));
-		AddToRadioSelection("Sync to chosen promoted label", SNew(SPickPromoted).Title(FText::FromString("Pick promoted label: ")));
-		AddToRadioSelection("Sync to chosen promotable label since last promoted", SNew(SPickPromotable).Title(FText::FromString("Pick promotable label: ")));
-		AddToRadioSelection("Sync to any chosen label", SNew(SPickAny).Title(FText::FromString("Pick label: ")));
+		AddToRadioSelection(LOCTEXT("SyncToLatestPromoted", "Sync to the latest promoted"), SNew(SLatestPromoted));
+		AddToRadioSelection(LOCTEXT("SyncToChosenPromoted", "Sync to chosen promoted label"), SNew(SPickPromoted).Title(LOCTEXT("PickPromotedLabel", "Pick promoted label: ")));
+		AddToRadioSelection(LOCTEXT("SyncToChosenPromotable", "Sync to chosen promotable label since last promoted"), SNew(SPickPromotable).Title(LOCTEXT("PickPromotableLabel", "Pick promotable label: ")));
+		AddToRadioSelection(LOCTEXT("SyncToAnyChosen", "Sync to any chosen label"), SNew(SPickAny).Title(LOCTEXT("PickLabel", "Pick label: ")));
 
 		PickGameWidget = SNew(SPickGameWidget)
 			.OnGamePicked(this, &SMainTabWidget::OnCurrentGameChanged);
 
 		ArtistSyncCheckBox = SNew(SCheckBox)
 			[
-				SNew(STextBlock).Text(FText::FromString("Artist sync?"))
+				SNew(STextBlock).Text(LOCTEXT("ArtistSync", "Artist sync?"))
 			];
 
 #if UE_BUILD_DEBUG
@@ -1073,14 +1065,19 @@ public:
 
 		PreviewSyncCheckBox = SNew(SCheckBox)
 			[
-				SNew(STextBlock).Text(FText::FromString("Preview sync?"))
+				SNew(STextBlock).Text(LOCTEXT("PreviewSync", "Preview sync?"))
+			];
+
+		AutoClobberSyncCheckBox = SNew(SCheckBox)
+			[
+				SNew(STextBlock).Text(LOCTEXT("AutoClobberSync", "Auto-clobber sync?"))
 			];
 
 		GoBackButton = SNew(SButton)
 			.IsEnabled(false)
 			.OnClicked(this, &SMainTabWidget::OnGoBackButtonClick)
 			[
-				SNew(STextBlock).Text(FText::FromString("Go back"))
+				SNew(STextBlock).Text(LOCTEXT("GoBack", "Go back"))
 			];
 
 		TSharedPtr<SVerticalBox> MainBox;
@@ -1111,6 +1108,14 @@ public:
 						]
 					]
 					+ SHorizontalBox::Slot().HAlign(HAlign_Fill)
+					+ SHorizontalBox::Slot().AutoWidth().VAlign(VAlign_Top).Padding(0.0f, 0.0f, 10.0f, 0.0f)
+					[
+						SNew(SVerticalBox)
+						+SVerticalBox::Slot()
+						[
+							AutoClobberSyncCheckBox.ToSharedRef()
+						]
+					]
 					+ SHorizontalBox::Slot().AutoWidth()
 					[
 						SNew(SVerticalBox)
@@ -1119,9 +1124,9 @@ public:
 							ArtistSyncCheckBox.ToSharedRef()
 						]
 						+ SVerticalBox::Slot()
-							[
-								PreviewSyncCheckBox.ToSharedRef()
-							]
+						[
+							PreviewSyncCheckBox.ToSharedRef()
+						]
 					]
 					+ SHorizontalBox::Slot().AutoWidth().Padding(30, 0, 0, 0)
 						[
@@ -1376,6 +1381,7 @@ private:
 		FSyncSettings Settings(
 			ArtistSyncCheckBox->IsChecked(),
 			PreviewSyncCheckBox->IsChecked(),
+			AutoClobberSyncCheckBox->IsChecked(),
 			OverrideSyncStep);
 
 		Switcher->SetActiveWidgetIndex(1);
@@ -1501,7 +1507,7 @@ private:
 	 * @param Widget Widget to add.
 	 */
 	template <typename TWidgetType>
-	void AddToRadioSelection(const FString& Name, TSharedRef<TWidgetType> Widget)
+	void AddToRadioSelection(FText Name, TSharedRef<TWidgetType> Widget)
 	{
 		RadioSelection->Add(Name, Widget);
 		SyncCommandLineProviders.Add(Widget);
@@ -1577,6 +1583,8 @@ private:
 	TSharedPtr<SCheckBox> ArtistSyncCheckBox;
 	/* Check box to tell if this should be a preview sync. */
 	TSharedPtr<SCheckBox> PreviewSyncCheckBox;
+	/* Check box to tell if this should be a auto-clobber sync. */
+	TSharedPtr<SCheckBox> AutoClobberSyncCheckBox;
 
 	/* External thread requests dispatcher. */
 	TSharedRef<FExternalThreadsDispatcher, ESPMode::ThreadSafe> ExternalThreadsDispatcher;
diff --git a/Engine/Source/Programs/UnrealSync/Private/UnrealSync.cpp b/Engine/Source/Programs/UnrealSync/Private/UnrealSync.cpp
--- a/Engine/Source/Programs/UnrealSync/Private/UnrealSync.cpp
+++ b/Engine/Source/Programs/UnrealSync/Private/UnrealSync.cpp
@@ -4,6 +4,7 @@
 #include "GUI.h"
 #include "P4DataCache.h"
 #include "P4Env.h"
+#include "XmlParser.h"
 
 #include "RequiredProgramMainCPPInclude.h"
 
@@ -402,7 +403,22 @@ void SyncingMessage(const FUnrealSync::FOnSyncProgress& OnSyncProgress, const FS
 	}
 }
 
-#include "XmlParser.h"
+/**
+ * Contains file and revision specification as sync step.
+ */
+struct FSyncStep
+{
+	FSyncStep(FString FileSpec, FString RevSpec)
+		: FileSpec(MoveTemp(FileSpec)), RevSpec(MoveTemp(RevSpec))
+	{}
+
+	const FString& GetFileSpec() const { return FileSpec; }
+	const FString& GetRevSpec() const { return RevSpec; }
+
+private:
+	FString FileSpec;
+	FString RevSpec;
+};
 
 /**
  * Parses sync rules XML content and fills array with sync steps
@@ -413,7 +429,7 @@ void SyncingMessage(const FUnrealSync::FOnSyncProgress& OnSyncProgress, const FS
  * @param ContentRevisionSpec Content revision spec.
  * @param ProgramRevisionSpec Program revision spec.
  */
-void FillWithSyncSteps(TArray<FString>& SyncSteps, const FString& SyncRules, const FString& ContentRevisionSpec, const FString& ProgramRevisionSpec)
+void FillWithSyncSteps(TArray<FSyncStep>& SyncSteps, const FString& SyncRules, const FString& ContentRevisionSpec, const FString& ProgramRevisionSpec)
 {
 	TSharedRef<FXmlFile> Doc = MakeShareable(new FXmlFile(SyncRules, EConstructMethod::ConstructFromBuffer));
 
@@ -443,7 +459,17 @@ void FillWithSyncSteps(TArray<FString>& SyncSteps, const FString& SyncRules, con
 				SyncStep += ContentRevisionSpec;
 			}
 
-			SyncSteps.Add(SyncStep);
+			int32 AtLastPos = INDEX_NONE;
+			SyncStep.FindLastChar('@', AtLastPos);
+
+			int32 HashLastPos = INDEX_NONE;
+			SyncStep.FindLastChar('#', HashLastPos);
+
+			int32 RevSpecSeparatorPos = FMath::Max<int32>(AtLastPos, HashLastPos);
+
+			check(RevSpecSeparatorPos != INDEX_NONE); // At least one rev specifier needed here.
+			
+			SyncSteps.Add(FSyncStep(SyncStep.Mid(0, RevSpecSeparatorPos), SyncStep.Mid(RevSpecSeparatorPos)));
 		}
 	}
 }
@@ -477,7 +503,7 @@ bool FUnrealSync::Sync(const FSyncSettings& Settings, const FString& Label, cons
 	SyncingMessage(OnSyncProgress, "Syncing to label: " + Label);
 
 	auto ProgramRevisionSpec = "@" + Label;
-	TArray<FString> SyncSteps;
+	TArray<FSyncStep> SyncSteps;
 
 	if (Settings.bArtist)
 	{
@@ -500,7 +526,7 @@ bool FUnrealSync::Sync(const FSyncSettings& Settings, const FString& Label, cons
 			*FP4Env::Get().GetBranch(), Game.IsEmpty() ? TEXT("Samples") : *Game);
 
 		FString SyncRules;
-		if (!FP4Env::RunP4Output("print -q " + ArtistSyncRulesPath + "#head", SyncRules))
+		if (!FP4Env::RunP4Output("print -q " + ArtistSyncRulesPath + "#head", SyncRules) || SyncRules.IsEmpty())
 		{
 			return false;
 		}
@@ -511,15 +537,74 @@ bool FUnrealSync::Sync(const FSyncSettings& Settings, const FString& Label, cons
 	}
 	else
 	{
-		SyncSteps.Add((Settings.OverrideSyncStep.IsEmpty() ? "/..." : Settings.OverrideSyncStep) + ProgramRevisionSpec); // all files to label
+		SyncSteps.Add(FSyncStep(Settings.OverrideSyncStep.IsEmpty() ? "/..." : Settings.OverrideSyncStep, ProgramRevisionSpec)); // all files to label
 	}
 
+	class FSyncCollectAndPassThrough
+	{
+	public:
+		FSyncCollectAndPassThrough(const FOnSyncProgress& Progress)
+			: Progress(Progress)
+		{ }
+
+		operator FOnSyncProgress() const
+		{
+			return FOnSyncProgress::CreateRaw(this, &FSyncCollectAndPassThrough::OnProgress);
+		}
+
+		bool OnProgress(const FString& Text)
+		{
+			Log += Text;
+
+			return Progress.Execute(Text);
+		}
+
+		void ProcessLog()
+		{
+			Log = Log.Replace(TEXT("\r"), TEXT(""));
+			const FRegexPattern CantClobber(TEXT("Can't clobber writable file ([^\\n]+)"));
+
+			FRegexMatcher Match(CantClobber, Log);
+
+			while (Match.FindNext())
+			{
+				CantClobbers.Add(Log.Mid(
+					Match.GetCaptureGroupBeginning(1),
+					Match.GetCaptureGroupEnding(1) - Match.GetCaptureGroupBeginning(1)));
+			}
+		}
+
+		const TArray<FString>& GetCantClobbers() const { return CantClobbers; }
+
+	private:
+		const FOnSyncProgress& Progress;
+
+		FString Log;
+
+		TArray<FString> CantClobbers;
+	};
+
 	FString CommandPrefix = FString("sync ") + (Settings.bPreview ? "-n " : "") + FP4Env::Get().GetBranch();
 	for(const auto& SyncStep : SyncSteps)
 	{
-		if (!FP4Env::RunP4Progress(CommandPrefix + SyncStep, OnSyncProgress))
+		FSyncCollectAndPassThrough ErrorsCollector(OnSyncProgress);
+		if (!FP4Env::RunP4Progress(CommandPrefix + SyncStep.GetFileSpec() + SyncStep.GetRevSpec(), ErrorsCollector))
 		{
-			return false;
+			if (!Settings.bAutoClobber)
+			{
+				return false;
+			}
+
+			ErrorsCollector.ProcessLog();
+
+			FString AutoClobberCommandPrefix("sync -f ");
+			for (const auto& CantClobber : ErrorsCollector.GetCantClobbers())
+			{
+				if (!FP4Env::RunP4Progress(AutoClobberCommandPrefix + CantClobber + SyncStep.GetRevSpec(), OnSyncProgress))
+				{
+					return false;
+				}
+			}
 		}
 	}
 
diff --git a/Engine/Source/Programs/UnrealSync/Private/UnrealSync.h b/Engine/Source/Programs/UnrealSync/Private/UnrealSync.h
--- a/Engine/Source/Programs/UnrealSync/Private/UnrealSync.h
+++ b/Engine/Source/Programs/UnrealSync/Private/UnrealSync.h
@@ -92,18 +92,23 @@ struct FSyncSettings
 	/* Preview sync? */
 	bool bPreview;
 
+	/* Auto-clobber sync? */
+	bool bAutoClobber;
+
 	/* Override sync step. If set then it overrides collecting sync steps and uses this one instead. */
 	FString OverrideSyncStep;
 
-	FSyncSettings(bool bArtist, bool bPreview, FString OverrideSyncStep = FString())
+	FSyncSettings(bool bArtist, bool bPreview, bool bAutoClobber, FString OverrideSyncStep = FString())
 		: bArtist(bArtist)
 		, bPreview(bPreview)
+		, bAutoClobber(bAutoClobber)
 		, OverrideSyncStep(MoveTemp(OverrideSyncStep))
 	{ }
 
 	FSyncSettings(const FSyncSettings&& Other)
-		: bArtist(MoveTemp(Other.bArtist))
-		, bPreview(MoveTemp(Other.bPreview))
+		: bArtist(Other.bArtist)
+		, bPreview(Other.bPreview)
+		, bAutoClobber(Other.bAutoClobber)
 		, OverrideSyncStep(MoveTemp(Other.OverrideSyncStep))
 	{ }
 };
