UnrealSync debug improvements.

diff --git a/Engine/Source/Programs/UnrealSync/Private/GUI.cpp b/Engine/Source/Programs/UnrealSync/Private/GUI.cpp
--- a/Engine/Source/Programs/UnrealSync/Private/GUI.cpp
+++ b/Engine/Source/Programs/UnrealSync/Private/GUI.cpp
@@ -1051,13 +1051,7 @@ public:
 				SNew(STextBlock).Text(LOCTEXT("ArtistSync", "Artist sync?"))
 			];
 
-#if UE_BUILD_DEBUG
-		auto bDebug = FParse::Param(FCommandLine::Get(), TEXT("Debug"));
-#else
-		const auto bDebug = false;
-#endif
-
-		if (!bDebug)
+		if (!FUnrealSync::IsDebugParameterSet())
 		{
 			// Checked by default.
 			ArtistSyncCheckBox->ToggleCheckedState();
@@ -1184,7 +1178,7 @@ public:
 				Switcher.ToSharedRef()
 			];
 
-		if (bDebug)
+		if (FUnrealSync::IsDebugParameterSet())
 		{
 			MainBox->InsertSlot(2).AutoHeight()
 				[
diff --git a/Engine/Source/Programs/UnrealSync/Private/UnrealSync.cpp b/Engine/Source/Programs/UnrealSync/Private/UnrealSync.cpp
--- a/Engine/Source/Programs/UnrealSync/Private/UnrealSync.cpp
+++ b/Engine/Source/Programs/UnrealSync/Private/UnrealSync.cpp
@@ -23,7 +23,8 @@ bool FUnrealSync::RunDetachedUS(const FString& USPath, bool bDoNotRunFromCopy, b
 	FString CommandLine = FString()
 		+ (bDoNotRunFromCopy ? TEXT("-DoNotRunFromCopy ") : TEXT(""))
 		+ (bDoNotUpdateOnStartUp ? TEXT("-DoNotUpdateOnStartUp ") : TEXT(""))
-		+ (bPassP4Env ? *FP4Env::Get().GetCommandLine()	: TEXT(""));
+		+ (FUnrealSync::IsDebugParameterSet() ? TEXT("-Debug ") : TEXT(""))
+		+ (bPassP4Env ? *FP4Env::Get().GetCommandLine() : TEXT(""));
 
 	return RunProcess(USPath, CommandLine);
 }
@@ -34,6 +35,7 @@ bool FUnrealSync::DeleteIfExistsAndCopyFile(const FString& To, const FString& Fr
 
 	if (PlatformPhysical.FileExists(*To) && !PlatformPhysical.DeleteFile(*To))
 	{
+		UE_LOG(LogUnrealSync, Warning, TEXT("Deleting existing file '%s' failed."), *To);
 		return false;
 	}
 
@@ -611,6 +613,13 @@ bool FUnrealSync::Sync(const FSyncSettings& Settings, const FString& Label, cons
 	return true;
 }
 
+bool FUnrealSync::IsDebugParameterSet()
+{
+	const auto bDebug = FParse::Param(FCommandLine::Get(), TEXT("Debug"));
+
+	return bDebug;
+}
+
 bool FUnrealSync::Initialization(const TCHAR* CommandLine)
 {
 	FString CommonExecutablePath =
diff --git a/Engine/Source/Programs/UnrealSync/Private/UnrealSync.h b/Engine/Source/Programs/UnrealSync/Private/UnrealSync.h
--- a/Engine/Source/Programs/UnrealSync/Private/UnrealSync.h
+++ b/Engine/Source/Programs/UnrealSync/Private/UnrealSync.h
@@ -130,6 +130,11 @@ public:
 	/* On sync log chunk read event delegate. */
 	DECLARE_DELEGATE_RetVal_OneParam(bool, FOnSyncProgress, const FString&);
 
+	/**
+	 * Tells if UnrealSync was run with -Debug param.
+	 */
+	static bool IsDebugParameterSet();
+
 	/**
 	 * This method copies UnrealSync to temp location and run it from there,
 	 * and updates (if needed).
