Added GetCaptureGroup to FRegexMatcher.

+ Changed old usages in UnrealSync.

diff --git a/Engine/Source/Programs/UnrealSync/Private/P4DataCache.cpp b/Engine/Source/Programs/UnrealSync/Private/P4DataCache.cpp
--- a/Engine/Source/Programs/UnrealSync/Private/P4DataCache.cpp
+++ b/Engine/Source/Programs/UnrealSync/Private/P4DataCache.cpp
@@ -7,21 +7,6 @@
 #include "P4Env.h"
 #include "Internationalization/Regex.h"
 
-/**
- * Helper function to get capture group value.
- * Missing in FRegexMatcher implementation (maybe it should go there?).
- *
- * @param Source Source text that was used to match.
- * @param Matcher Current matcher state.
- * @param Id Capture group ID to retrieve.
- *
- * @returns Captured group text.
- */
-FString GetGroupText(const FString& Source, FRegexMatcher& Matcher, int32 Id)
-{
-	return Source.Mid(Matcher.GetCaptureGroupBeginning(Id), Matcher.GetCaptureGroupEnding(Id) - Matcher.GetCaptureGroupBeginning(Id));
-}
-
 bool FP4DataCache::LoadFromLog(const FString& UnrealSyncListLog)
 {
 	const FRegexPattern Pattern(TEXT("Label ([^ ]+) (\\d{4})/(\\d{2})/(\\d{2}) (\\d{2}):(\\d{2}):(\\d{2})")); // '.+\\n
@@ -31,14 +16,14 @@ bool FP4DataCache::LoadFromLog(const FString& UnrealSyncListLog)
 	while (Matcher.FindNext())
 	{
 		Labels.Add(FP4Label(
-			GetGroupText(UnrealSyncListLog, Matcher, 1),
+			Matcher.GetCaptureGroup(1),
 			FDateTime(
-				FCString::Atoi(*GetGroupText(UnrealSyncListLog, Matcher, 2)),
-				FCString::Atoi(*GetGroupText(UnrealSyncListLog, Matcher, 3)),
-				FCString::Atoi(*GetGroupText(UnrealSyncListLog, Matcher, 4)),
-				FCString::Atoi(*GetGroupText(UnrealSyncListLog, Matcher, 5)),
-				FCString::Atoi(*GetGroupText(UnrealSyncListLog, Matcher, 6)),
-				FCString::Atoi(*GetGroupText(UnrealSyncListLog, Matcher, 7))
+				FCString::Atoi(*Matcher.GetCaptureGroup(2)),
+				FCString::Atoi(*Matcher.GetCaptureGroup(3)),
+				FCString::Atoi(*Matcher.GetCaptureGroup(4)),
+				FCString::Atoi(*Matcher.GetCaptureGroup(5)),
+				FCString::Atoi(*Matcher.GetCaptureGroup(6)),
+				FCString::Atoi(*Matcher.GetCaptureGroup(7))
 			)));
 	}
 
diff --git a/Engine/Source/Programs/UnrealSync/Private/P4Env.cpp b/Engine/Source/Programs/UnrealSync/Private/P4Env.cpp
--- a/Engine/Source/Programs/UnrealSync/Private/P4Env.cpp
+++ b/Engine/Source/Programs/UnrealSync/Private/P4Env.cpp
@@ -677,7 +677,7 @@ public:
 			FRegexMatcher Matcher(UserNamePattern, InfoOutput);
 			if (Matcher.FindNext())
 			{
-				Output = InfoOutput.Mid(Matcher.GetCaptureGroupBeginning(1), Matcher.GetCaptureGroupEnding(1) - Matcher.GetCaptureGroupBeginning(1));
+				Output = Matcher.GetCaptureGroup(1);
 				return true;
 			}
 		}
@@ -741,8 +741,8 @@ public:
 
 		while (Matcher->FindNext())
 		{
-			auto ClientName = P4ClientsOutput.Mid(Matcher->GetCaptureGroupBeginning(1), Matcher->GetCaptureGroupEnding(1) - Matcher->GetCaptureGroupBeginning(1));
-			auto Root = P4ClientsOutput.Mid(Matcher->GetCaptureGroupBeginning(2), Matcher->GetCaptureGroupEnding(2) - Matcher->GetCaptureGroupBeginning(2));
+			auto ClientName = Matcher->GetCaptureGroup(1);
+			auto Root = Matcher->GetCaptureGroup(2);
 
 			if (KnownPath.StartsWith(FPaths::ConvertRelativePathToFull(Root)))
 			{
@@ -757,10 +757,7 @@ public:
 
 				while (InfoMatcher.FindNext())
 				{
-					if (InfoOutput.Mid(
-						InfoMatcher.GetCaptureGroupBeginning(1),
-						InfoMatcher.GetCaptureGroupEnding(1) - InfoMatcher.GetCaptureGroupBeginning(1)
-						).Equals(HostName, ESearchCase::IgnoreCase))
+					if (InfoMatcher.GetCaptureGroup(1).Equals(HostName, ESearchCase::IgnoreCase))
 					{
 						Output = ClientName;
 						return true;
@@ -1063,13 +1060,8 @@ bool FP4Env::CheckIfFileNeedsUpdate(const FString& FilePath)
 		return false;
 	}
 
-	int32 HeadRev = FPlatformString::Atoi(*Output.Mid(
-		HeadRevMatcher.GetCaptureGroupBeginning(1),
-		HeadRevMatcher.GetCaptureGroupEnding(1) - HeadRevMatcher.GetCaptureGroupBeginning(1)));
-
-	int32 HaveRev = FPlatformString::Atoi(*Output.Mid(
-		HaveRevMatcher.GetCaptureGroupBeginning(1),
-		HaveRevMatcher.GetCaptureGroupEnding(1) - HaveRevMatcher.GetCaptureGroupBeginning(1)));
+	int32 HeadRev = FPlatformString::Atoi(*HeadRevMatcher.GetCaptureGroup(1));
+	int32 HaveRev = FPlatformString::Atoi(*HaveRevMatcher.GetCaptureGroup(1));
 
 	return HaveRev < HeadRev;
 }
diff --git a/Engine/Source/Programs/UnrealSync/Private/UnrealSync.cpp b/Engine/Source/Programs/UnrealSync/Private/UnrealSync.cpp
--- a/Engine/Source/Programs/UnrealSync/Private/UnrealSync.cpp
+++ b/Engine/Source/Programs/UnrealSync/Private/UnrealSync.cpp
@@ -570,9 +570,7 @@ bool FUnrealSync::Sync(const FSyncSettings& Settings, const FString& Label, cons
 
 			while (Match.FindNext())
 			{
-				CantClobbers.Add(Log.Mid(
-					Match.GetCaptureGroupBeginning(1),
-					Match.GetCaptureGroupEnding(1) - Match.GetCaptureGroupBeginning(1)));
+				CantClobbers.Add(Match.GetCaptureGroup(1));
 			}
 		}
 
diff --git a/Engine/Source/Runtime/Core/Private/Internationalization/ICURegex.cpp b/Engine/Source/Runtime/Core/Private/Internationalization/ICURegex.cpp
--- a/Engine/Source/Runtime/Core/Private/Internationalization/ICURegex.cpp
+++ b/Engine/Source/Runtime/Core/Private/Internationalization/ICURegex.cpp
@@ -46,13 +46,14 @@ namespace
 class FRegexMatcherImplementation
 {
 public:
-	FRegexMatcherImplementation(const FRegexPatternImplementation& Pattern, const FString& InputString) : ICUInputString( ICUUtilities::ConvertString(InputString) ), ICURegexMatcher( CreateRegexMatcher(Pattern, ICUInputString) )
+	FRegexMatcherImplementation(const FRegexPatternImplementation& Pattern, const FString& InputString) : ICUInputString(ICUUtilities::ConvertString(InputString)), ICURegexMatcher(CreateRegexMatcher(Pattern, ICUInputString)), OriginalString(InputString)
 	{
 	}
 
 public:
 	const icu::UnicodeString ICUInputString;
 	TSharedRef<icu::RegexMatcher> ICURegexMatcher;
+	FString OriginalString;
 };
 
 FRegexMatcher::FRegexMatcher(const FRegexPattern& Pattern, const FString& InputString) : Implementation(new FRegexMatcherImplementation(Pattern.Implementation.Get(), InputString))
@@ -88,6 +89,11 @@ int32 FRegexMatcher::GetCaptureGroupEnding(const int32 Index)
 	return Implementation->ICURegexMatcher->end(Index, ICUStatus);
 }
 
+FString FRegexMatcher::GetCaptureGroup(const int32 Index)
+{
+	return Implementation->OriginalString.Mid(GetCaptureGroupBeginning(Index), GetCaptureGroupEnding(Index) - GetCaptureGroupBeginning(Index));
+}
+
 int32 FRegexMatcher::GetBeginLimit()
 {
 	return Implementation->ICURegexMatcher->regionStart();
diff --git a/Engine/Source/Runtime/Core/Private/Internationalization/LegacyRegex.cpp b/Engine/Source/Runtime/Core/Private/Internationalization/LegacyRegex.cpp
--- a/Engine/Source/Runtime/Core/Private/Internationalization/LegacyRegex.cpp
+++ b/Engine/Source/Runtime/Core/Private/Internationalization/LegacyRegex.cpp
@@ -48,6 +48,11 @@ int32 FRegexMatcher::GetCaptureGroupEnding(const int32 Index)
 	return INDEX_NONE;
 }
 
+FString FRegexMatcher::GetCaptureGroup(const int32 Index)
+{
+	return FString();
+}
+
 int32 FRegexMatcher::GetBeginLimit()
 {
 	return INDEX_NONE;
diff --git a/Engine/Source/Runtime/Core/Public/Internationalization/Regex.h b/Engine/Source/Runtime/Core/Public/Internationalization/Regex.h
--- a/Engine/Source/Runtime/Core/Public/Internationalization/Regex.h
+++ b/Engine/Source/Runtime/Core/Public/Internationalization/Regex.h
@@ -27,6 +27,7 @@ public:
 
 	int32 GetCaptureGroupBeginning(const int32 Index);
 	int32 GetCaptureGroupEnding(const int32 Index);
+	FString GetCaptureGroup(const int32 Index);
 
 	int32 GetBeginLimit();
 	int32 GetEndLimit();
