UE-14029: Project selections are out of date in Unreal Sync

Changed US3 to ignore case when searching for projects.

diff --git a/Engine/Source/Programs/UnrealSync/Private/UnrealSync.cpp b/Engine/Source/Programs/UnrealSync/Private/UnrealSync.cpp
--- a/Engine/Source/Programs/UnrealSync/Private/UnrealSync.cpp
+++ b/Engine/Source/Programs/UnrealSync/Private/UnrealSync.cpp
@@ -211,7 +211,6 @@ TSharedPtr<TArray<FString> > FUnrealSync::GetAllLabels()
 
 TSharedPtr<TArray<FString> > FUnrealSync::GetPossibleGameNames()
 {
-	/* TODO: Hard coded game names. Needs to be fixed! */
 	TSharedPtr<TArray<FString> > PossibleGames = MakeShareable(new TArray<FString>());
 
 	FP4Env& Env = FP4Env::Get();
@@ -225,7 +224,7 @@ TSharedPtr<TArray<FString> > FUnrealSync::GetPossibleGameNames()
 	FString Line, Rest = FileList;
 	while (Rest.Split(LINE_TERMINATOR, &Line, &Rest, ESearchCase::CaseSensitive))
 	{
-		if (!Line.StartsWith(Env.GetBranch(), ESearchCase::CaseSensitive))
+		if (!Line.StartsWith(Env.GetBranch()))
 		{
 			continue;
 		}
