UECORE-115: Add copy/save/clear functionality for Unreal Sync logging

diff --git a/Engine/Source/Programs/UnrealSync/Private/GUI.cpp b/Engine/Source/Programs/UnrealSync/Private/GUI.cpp
--- a/Engine/Source/Programs/UnrealSync/Private/GUI.cpp
+++ b/Engine/Source/Programs/UnrealSync/Private/GUI.cpp
@@ -15,6 +15,8 @@
 #include "SThrobber.h"
 #include "SMultiLineEditableTextBox.h"
 #include "BaseTextLayoutMarshaller.h"
+#include "DesktopPlatformModule.h"
+#include "IDesktopPlatform.h"
 
 #define LOCTEXT_NAMESPACE "UnrealSync"
 
@@ -70,6 +72,14 @@ public:
 		MessagesTextBox->Refresh();
 	}
 
+	/**
+	 * Gets text version of the log.
+	 */
+	const FText GetPlainText()
+	{
+		return MessagesTextBox->GetPlainText();
+	}
+
 private:
 	/**
 	 * Messages marshaller.
@@ -1178,7 +1188,15 @@ public:
 						SAssignNew(CancelButton, SButton)
 						.OnClicked(this, &SMainTabWidget::OnCancelButtonClick)
 						[
-							SNew(STextBlock).Text(FText::FromString("Cancel"))
+							SNew(STextBlock).Text(LOCTEXT("Cancel", "Cancel"))
+						]
+					]
+					+ SHorizontalBox::Slot().AutoWidth()
+					[
+						SNew(SButton)
+						.OnClicked(this, &SMainTabWidget::OnSaveLogButtonClick)
+						[
+							SNew(STextBlock).Text(LOCTEXT("SaveLog", "Save Log..."))
 						]
 					]
 					+ SHorizontalBox::Slot().AutoWidth()
@@ -1285,6 +1303,63 @@ private:
 		return *SyncCommandLineProviders[RadioSelection->GetChosen()];
 	}
 
+	/**
+	 * Saves the log to the user-pointed location.
+	 */
+	void SaveLog()
+	{
+		IDesktopPlatform* DesktopPlatform = FDesktopPlatformModule::Get();
+
+		if (DesktopPlatform != NULL)
+		{
+			TArray<FString> Filenames;
+
+			TSharedPtr<SWindow> ParentWindow = FSlateApplication::Get().FindWidgetWindow(AsShared());
+			void* ParentWindowHandle = (ParentWindow.IsValid() && ParentWindow->GetNativeWindow().IsValid()) ? ParentWindow->GetNativeWindow()->GetOSWindowHandle() : nullptr;
+
+			if (DesktopPlatform->SaveFileDialog(
+				ParentWindowHandle,
+				LOCTEXT("SaveLogDialogTitle", "Save Log As...").ToString(),
+				FPaths::Combine(*FPaths::EngineSavedDir(), TEXT("Logs")),
+				TEXT("Syncing.log"),
+				TEXT("Log Files (*.log)|*.log"),
+				EFileDialogFlags::None,
+				Filenames))
+			{
+				if (Filenames.Num() > 0)
+				{
+					FString Filename = Filenames[0];
+
+					// add a file extension if none was provided
+					if (FPaths::GetExtension(Filename).IsEmpty())
+					{
+						Filename += Filename + TEXT(".log");
+					}
+
+					// save file
+					FArchive* LogFile = IFileManager::Get().CreateFileWriter(*Filename);
+
+					if (LogFile != nullptr)
+					{
+						FString LogText = Log->GetPlainText().ToString() + LINE_TERMINATOR;
+						LogFile->Serialize(TCHAR_TO_ANSI(*LogText), LogText.Len());
+						LogFile->Close();
+
+						delete LogFile;
+					}
+					else
+					{
+						FMessageDialog::Open(EAppMsgType::Ok, LOCTEXT("SaveLogDialogFileError", "Failed to open the specified file for saving!"));
+					}
+				}
+			}
+		}
+		else
+		{
+			FMessageDialog::Open(EAppMsgType::Ok, LOCTEXT("SaveLogDialogUnsupportedError", "Saving is not supported on this platform!"));
+		}
+	}
+
 	/**
 	 * Performs a procedure that needs to be done when GoBack button is clicked.
 	 * It tells the app to go back to the beginning.
@@ -1309,6 +1384,18 @@ private:
 		return FReply::Handled();
 	}
 
+	/**
+	 * Function that is called on when save log button is clicked.
+	 *
+	 * @returns Tells that this event was handled.
+	 */
+	FReply OnSaveLogButtonClick()
+	{
+		SaveLog();
+
+		return FReply::Handled();
+	}
+
 	/**
 	 * Function that is called on when cancel button is clicked.
 	 * It tells the app to cancel the sync and go back to the beginning.
diff --git a/Engine/Source/Programs/UnrealSync/UnrealSync.Build.cs b/Engine/Source/Programs/UnrealSync/UnrealSync.Build.cs
--- a/Engine/Source/Programs/UnrealSync/UnrealSync.Build.cs
+++ b/Engine/Source/Programs/UnrealSync/UnrealSync.Build.cs
@@ -18,7 +18,8 @@ public class UnrealSync : ModuleRules
 				"SlateCore",
 				"StandaloneRenderer",
 				"XmlParser",
-				"JSON"
+				"JSON",
+				"DesktopPlatform"
 			}
 		);
 		
