UE-14034: UnrealSync resets to default settings when opened

diff --git a/Engine/Source/Programs/UnrealSync/Private/GUI.cpp b/Engine/Source/Programs/UnrealSync/Private/GUI.cpp
--- a/Engine/Source/Programs/UnrealSync/Private/GUI.cpp
+++ b/Engine/Source/Programs/UnrealSync/Private/GUI.cpp
@@ -1,6 +1,7 @@
 // Copyright 1998-2015 Epic Games, Inc. All Rights Reserved.
 
 #include "UnrealSync.h"
+#include "SyncingThread.h"
 
 #include "SlateBasics.h"
 #include "StandaloneRenderer.h"
@@ -17,6 +18,7 @@
 #include "BaseTextLayoutMarshaller.h"
 #include "DesktopPlatformModule.h"
 #include "IDesktopPlatform.h"
+#include "SettingsCache.h"
 
 #define LOCTEXT_NAMESPACE "UnrealSync"
 
@@ -142,6 +144,55 @@ private:
 	int32 NumberOfLines = 0;
 };
 
+/**
+ * A checkbox with JSON preservable state.
+ */
+class SPreservableCheckBox : public SCheckBox
+{
+public:
+	SLATE_BEGIN_ARGS(SPreservableCheckBox)	{}
+		SLATE_DEFAULT_SLOT(FArguments, Content)
+	SLATE_END_ARGS()
+
+	BEGIN_SLATE_FUNCTION_BUILD_OPTIMIZATION
+	void Construct(const FArguments& InArgs)
+	{
+		SCheckBox::FArguments BaseArgs;
+		BaseArgs._Content = InArgs._Content;
+
+		SCheckBox::Construct(BaseArgs);
+	}
+	END_SLATE_FUNCTION_BUILD_OPTIMIZATION
+
+	/**
+	 * Converts current object state to JSON value.
+	 *
+	 * @returns JSON value object with current state.
+	 */
+	TSharedRef<FJsonValue> GetStateAsJSON() const
+	{
+		check(GetCheckedState() != ECheckBoxState::Undetermined);
+
+		return TSharedRef<FJsonValue>(new FJsonValueBoolean(GetCheckedState() == ECheckBoxState::Checked));
+	}
+
+	/**
+	 * Loads state from JSON value object.
+	 *
+	 * @param Value JSON value object.
+	 */
+	void LoadStateFromJSON(const FJsonValue& Value)
+	{
+		bool BoolValue;
+		check(Value.TryGetBool(BoolValue) && GetCheckedState() != ECheckBoxState::Undetermined);
+
+		if ((GetCheckedState() == ECheckBoxState::Checked) != BoolValue)
+		{
+			ToggleCheckedState();
+		}
+	}
+};
+
 /**
  * Performs a dialog when running editor is detected.
  */
@@ -347,6 +398,22 @@ public:
 		return *OptionsSource[CurrentOption];
 	}
 
+	/**
+	 * Sets current option for this combo box.
+	 *
+	 * @param Value Value to set combo box to.
+	 */
+	void SetCurrentOption(const FString& Value)
+	{
+		for (int32 OptionId = 0; OptionId < OptionsSource.Num(); ++OptionId)
+		{
+			if (OptionsSource[OptionId]->Equals(Value))
+			{
+				CurrentOption = OptionId;
+			}
+		}
+	}
+
 	/**
 	 * Checks if options source is empty.
 	 *
@@ -376,13 +443,7 @@ private:
 	void ComboBoxSelectionChanged(TSharedPtr<FString> Value, ESelectInfo::Type SelectInfo)
 	{
 		CurrentOption = -1;
-		for (int32 OptionId = 0; OptionId < OptionsSource.Num(); ++OptionId)
-		{
-			if (OptionsSource[OptionId]->Equals(*Value))
-			{
-				CurrentOption = OptionId;
-			}
-		}
+		SetCurrentOption(*Value);
 
 		OnSelectionChanged.ExecuteIfBound(CurrentOption, SelectInfo);
 	}
@@ -621,6 +682,31 @@ public:
 		return Items[ItemId]->GetContent();
 	}
 
+	/**
+	 * Converts current object state to JSON value.
+	 *
+	 * @returns JSON value object with current state.
+	 */
+	TSharedRef<FJsonValue> GetStateAsJSON() const
+	{
+		return MakeShareable(new FJsonValueNumber(GetChosen()));
+	}
+
+	/**
+	 * Loads state from JSON value object.
+	 *
+	 * @param Value JSON value object.
+	 */
+	void LoadStateFromJSON(const FJsonValue& Value)
+	{
+		int32 Chosen = 0;
+
+		if (Value.TryGetNumber(Chosen) && Chosen < Items.Num())
+		{
+			ChooseEnabledItem(Chosen);
+		}
+	}
+
 private:
 	/**
 	 * Function to rebuild main and fill it with items provided.
@@ -674,8 +760,10 @@ public:
 	}
 
 	BEGIN_SLATE_FUNCTION_BUILD_OPTIMIZATION
-	void Construct(const FArguments& InArgs)
+	void Construct(const FArguments& InArgs, const TSharedRef<FP4DataProxy>& InData)
 	{
+		Data = InData;
+
 		this->ChildSlot
 			[
 				SNew(STextBlock).Text(FText::FromString("This option will sync to the latest promoted label for given game."))
@@ -727,7 +815,7 @@ public:
 	{
 		ILabelNameProvider::RefreshData(GameName);
 
-		auto Labels = FUnrealSync::GetPromotedLabelsForGame(GameName);
+		auto Labels = Data.Pin()->GetPromotedLabelsForGame(GameName);
 		if (Labels->Num() != 0)
 		{
 			CurrentLabelName = (*Labels)[0];
@@ -758,6 +846,26 @@ public:
 		return true;
 	}
 
+	/**
+	 * Converts current object state to JSON value.
+	 *
+	 * @returns JSON value object with current state.
+	 */
+	TSharedPtr<FJsonValue> GetStateAsJSON() const
+	{
+		return nullptr;
+	}
+
+	/**
+	 * Loads state from JSON value object.
+	 *
+	 * @param Value JSON value object.
+	 */
+	void LoadStateFromJSON(const FJsonValue& Value)
+	{
+
+	}
+
 private:
 	/* Currently picked game name. */
 	FString CurrentGameName;
@@ -767,6 +875,9 @@ private:
 
 	/* Event that will trigger when data is ready. */
 	TSharedPtr<FEvent> DataReady;
+
+	/* Reference to P4 data proxy. */
+	TWeakPtr<FP4DataProxy> Data;
 };
 
 /**
@@ -780,9 +891,11 @@ public:
 	SLATE_END_ARGS()
 
 	BEGIN_SLATE_FUNCTION_BUILD_OPTIMIZATION
-	void Construct(const FArguments& InArgs)
+	void Construct(const FArguments& InArgs, const TSharedRef<FP4DataProxy>& InData)
 	{
-		LabelsCombo = SNew(SSimpleTextComboBox);
+		Data = InData;
+		LabelsCombo = SNew(SSimpleTextComboBox)
+			.OnSelectionChanged(this, &SPickLabel::OnSelectionChanged);
 		LabelsCombo->Add("Needs refreshing");
 
 		this->ChildSlot
@@ -817,7 +930,7 @@ public:
 	 */
 	bool IsReadyForSync() const override
 	{
-		return FUnrealSync::HasValidData() && !LabelsCombo->IsEmpty();
+		return Data.Pin()->HasValidData() && !LabelsCombo->IsEmpty();
 	}
 
 	/**
@@ -834,6 +947,31 @@ public:
 		LabelsCombo->SetEnabled(false);
 	}
 
+	/**
+	 * Converts current object state to JSON value.
+	 *
+	 * @returns JSON value object with current state.
+	 */
+	TSharedPtr<FJsonValue> GetStateAsJSON() const
+	{
+		if (!LabelsCombo->IsEnabled())
+		{
+			return nullptr;
+		}
+
+		return MakeShareable(new FJsonValueString(GetLabelName()));
+	}
+
+	/**
+	 * Loads state from JSON value object.
+	 *
+	 * @param Value JSON value object.
+	 */
+	void LoadStateFromJSON(const FJsonValue& Value)
+	{
+		Value.TryGetString(LastPickedValue);
+	}
+
 protected:
 	/**
 	 * Reset current labels combo options.
@@ -844,7 +982,7 @@ protected:
 	{
 		LabelsCombo->Clear();
 
-		if (FUnrealSync::HasValidData())
+		if (Data.Pin()->HasValidData())
 		{
 			for (auto LabelName : LabelOptions)
 			{
@@ -852,6 +990,11 @@ protected:
 			}
 
 			LabelsCombo->SetEnabled(true);
+
+			if (!LastPickedValue.IsEmpty())
+			{
+				LabelsCombo->SetCurrentOption(LastPickedValue);
+			}
 		}
 		else
 		{
@@ -860,9 +1003,29 @@ protected:
 		}
 	}
 
+	/**
+	 * Gets current data proxy.
+	 */
+	FP4DataProxy& GetData() { return *Data.Pin().Get(); }
+
 private:
+	/**
+	 * Function is called when combo box selection is changed.
+	 * It called OnGamePicked event.
+	 */
+	void OnSelectionChanged(int32 SelectionId, ESelectInfo::Type SelectionInfo)
+	{
+		LastPickedValue = LabelsCombo->GetCurrentOption();
+	}
+
 	/* Labels combo widget. */
 	TSharedPtr<SSimpleTextComboBox> LabelsCombo;
+
+	/* Reference to P4 data proxy. */
+	TWeakPtr<FP4DataProxy> Data;
+
+	/* Last picked value. */
+	FString LastPickedValue;
 };
 
 /**
@@ -878,7 +1041,7 @@ public:
 	 */
 	void RefreshData(const FString& GameName) override
 	{
-		SetLabelOptions(*FUnrealSync::GetPromotedLabelsForGame(GameName));
+		SetLabelOptions(*GetData().GetPromotedLabelsForGame(GameName));
 
 		SPickLabel::RefreshData(GameName);
 	}
@@ -897,7 +1060,7 @@ public:
 	 */
 	void RefreshData(const FString& GameName) override
 	{
-		SetLabelOptions(*FUnrealSync::GetPromotableLabelsForGame(GameName));
+		SetLabelOptions(*GetData().GetPromotableLabelsForGame(GameName));
 
 		SPickLabel::RefreshData(GameName);
 	}
@@ -916,7 +1079,7 @@ public:
 	 */
 	void RefreshData(const FString& GameName) override
 	{
-		SetLabelOptions(*FUnrealSync::GetAllLabels());
+		SetLabelOptions(*GetData().GetAllLabels());
 
 		SPickLabel::RefreshData(GameName);
 	}
@@ -936,14 +1099,15 @@ public:
 	SLATE_END_ARGS()
 
 	BEGIN_SLATE_FUNCTION_BUILD_OPTIMIZATION
-	void Construct(const FArguments& InArgs)
+	void Construct(const FArguments& InArgs, const TSharedRef<FP4DataProxy>& InData)
 	{
+		Data = InData;
 		OnGamePicked = InArgs._OnGamePicked;
 
 		GamesOptions = SNew(SSimpleTextComboBox)
 			.OnSelectionChanged(this, &SPickGameWidget::OnComboBoxSelectionChanged);
 
-		auto PossibleGameNames = FUnrealSync::GetPossibleGameNames();
+		auto PossibleGameNames = Data.Pin()->GetPossibleGameNames();
 
 		if (PossibleGameNames->Remove(FUnrealSync::GetSharedPromotableP4FolderName()))
 		{
@@ -975,11 +1139,34 @@ public:
 	 *
 	 * @returns Currently selected game.
 	 */
-	const FString& GetCurrentGame()
+	const FString& GetCurrentGame() const
 	{
 		return GamesOptions->GetCurrentOption();
 	}
 
+	/**
+	 * Converts current object state to JSON value.
+	 *
+	 * @returns JSON value object with current state.
+	 */
+	TSharedRef<FJsonValue> GetStateAsJSON() const
+	{
+		return TSharedRef<FJsonValueString>(new FJsonValueString(GetCurrentGame()));
+	}
+
+	/**
+	 * Loads state from JSON value object.
+	 *
+	 * @param Value JSON value object.
+	 */
+	void LoadStateFromJSON(const FJsonValue& Value)
+	{
+		FString StrVal;
+		check(Value.TryGetString(StrVal));
+
+		GamesOptions->SetCurrentOption(StrVal);
+	}
+
 private:
 	/**
 	 * Function is called when combo box selection is changed.
@@ -994,6 +1181,9 @@ private:
 	TSharedPtr<SSimpleTextComboBox> GamesOptions;
 	/* Delegate that is going to be called when game was picked by the user. */
 	FOnGamePicked OnGamePicked;
+
+	/* Reference to P4 data proxy. */
+	TWeakPtr<FP4DataProxy> Data;
 };
 
 /**
@@ -1051,9 +1241,25 @@ class SMainTabWidget : public SCompoundWidget
 	};
 
 public:
+	/* On data loaded event delegate. */
+	DECLARE_DELEGATE(FOnDataLoaded);
+
+	/* On data reset event delegate. */
+	DECLARE_DELEGATE(FOnDataReset);
+
 	SMainTabWidget()
-		: ExternalThreadsDispatcher(MakeShareable(new FExternalThreadsDispatcher()))
-	{}
+		: ExternalThreadsDispatcher(MakeShareable(new FExternalThreadsDispatcher())), P4Data(new FP4DataProxy)
+	{
+		GetSingletonPtr() = StaticCastSharedRef<SMainTabWidget>(this->AsShared());
+	}
+
+	/**
+	 * Gets instance of this class.
+	 */
+	static TSharedPtr<SMainTabWidget> GetInstance()
+	{
+		return GetSingletonPtr();
+	}
 
 	SLATE_BEGIN_ARGS(SMainTabWidget) {}
 	SLATE_END_ARGS()
@@ -1063,15 +1269,15 @@ public:
 	{
 		RadioSelection = SNew(SRadioContentSelection);
 
-		AddToRadioSelection(LOCTEXT("SyncToLatestPromoted", "Sync to the latest promoted"), SNew(SLatestPromoted));
-		AddToRadioSelection(LOCTEXT("SyncToChosenPromoted", "Sync to chosen promoted label"), SNew(SPickPromoted).Title(LOCTEXT("PickPromotedLabel", "Pick promoted label: ")));
-		AddToRadioSelection(LOCTEXT("SyncToChosenPromotable", "Sync to chosen promotable label since last promoted"), SNew(SPickPromotable).Title(LOCTEXT("PickPromotableLabel", "Pick promotable label: ")));
-		AddToRadioSelection(LOCTEXT("SyncToAnyChosen", "Sync to any chosen label"), SNew(SPickAny).Title(LOCTEXT("PickLabel", "Pick label: ")));
+		AddToRadioSelection(LOCTEXT("SyncToLatestPromoted", "Sync to the latest promoted"), SAssignNew(LatestPromoted, SLatestPromoted, P4Data));
+		AddToRadioSelection(LOCTEXT("SyncToChosenPromoted", "Sync to chosen promoted label"), SAssignNew(ChosenPromoted, SPickPromoted, P4Data).Title(LOCTEXT("PickPromotedLabel", "Pick promoted label: ")));
+		AddToRadioSelection(LOCTEXT("SyncToChosenPromotable", "Sync to chosen promotable label since last promoted"), SAssignNew(ChosenPromotable, SPickPromotable, P4Data).Title(LOCTEXT("PickPromotableLabel", "Pick promotable label: ")));
+		AddToRadioSelection(LOCTEXT("SyncToAnyChosen", "Sync to any chosen label"), SAssignNew(ChosenLabel, SPickAny, P4Data).Title(LOCTEXT("PickLabel", "Pick label: ")));
 
-		PickGameWidget = SNew(SPickGameWidget)
+		PickGameWidget = SNew(SPickGameWidget, P4Data)
 			.OnGamePicked(this, &SMainTabWidget::OnCurrentGameChanged);
 
-		ArtistSyncCheckBox = SNew(SCheckBox)
+		ArtistSyncCheckBox = SNew(SPreservableCheckBox)
 			[
 				SNew(STextBlock).Text(LOCTEXT("ArtistSync", "Artist sync?"))
 			];
@@ -1082,12 +1288,12 @@ public:
 			ArtistSyncCheckBox->ToggleCheckedState();
 		}
 
-		PreviewSyncCheckBox = SNew(SCheckBox)
+		PreviewSyncCheckBox = SNew(SPreservableCheckBox)
 			[
 				SNew(STextBlock).Text(LOCTEXT("PreviewSync", "Preview sync?"))
 			];
 
-		AutoClobberSyncCheckBox = SNew(SCheckBox)
+		AutoClobberSyncCheckBox = SNew(SPreservableCheckBox)
 			[
 				SNew(STextBlock).Text(LOCTEXT("AutoClobberSync", "Auto-clobber sync?"))
 			];
@@ -1227,25 +1433,168 @@ public:
 				];
 		}
 
-		FUnrealSync::RegisterOnDataReset(FUnrealSync::FOnDataReset::CreateRaw(this, &SMainTabWidget::DataReset));
-		FUnrealSync::RegisterOnDataLoaded(FUnrealSync::FOnDataLoaded::CreateRaw(this, &SMainTabWidget::DataLoaded));
+		OnDataReset = FOnDataReset::CreateRaw(this, &SMainTabWidget::DataReset);
+		OnDataLoaded = FOnDataLoaded::CreateRaw(this, &SMainTabWidget::DataLoaded);
 
 		RegisterActiveTimer(0.5f, FWidgetActiveTimerDelegate::CreateThreadSafeSP(&ExternalThreadsDispatcher.Get(), &FExternalThreadsDispatcher::ExecuteRequests));
 
 		OnReloadLabels();
+
+		TSharedPtr<FJsonValue> State;
+		if (FSettingsCache::Get().GetSetting(State, TEXT("GUI")))
+		{
+			LoadStateFromJSON(*State.Get());
+		}
 	}
 	END_SLATE_FUNCTION_BUILD_OPTIMIZATION
 
+	/**
+	 * JSON serializer for SMainTabWidget.
+	 */
+	template <bool bIsConst>
+	struct FJsonSerializer
+	{
+		/* Const-correct parameter type depending on template parameter. */
+		typedef typename FSettingsCache::TAddConstIf<bIsConst, SMainTabWidget&>::Result ConstCorrectType;
+
+		/**
+		 * Serializes SMainTabWidget object for JSON using given TSerializationActionPolicy.
+		 *
+		 * @param Action Action object to perform serialization on.
+		 * @param Object Object to serialize.
+		 */
+		template <class TSerializationActionPolicy>
+		static void Serialize(TSerializationActionPolicy& Action, ConstCorrectType Object)
+		{
+			Action.Serialize("PickedGame", Object.PickGameWidget);
+			Action.Serialize("IsArtist", Object.ArtistSyncCheckBox);
+			Action.Serialize("IsPreview", Object.PreviewSyncCheckBox);
+			Action.Serialize("IsAutoClobber", Object.AutoClobberSyncCheckBox);
+			Action.Serialize("RadioSelection", Object.RadioSelection);
+			Action.Serialize("LatestPromoted", Object.LatestPromoted);
+			Action.Serialize("ChosenPromoted", Object.ChosenPromoted);
+			Action.Serialize("ChosenPromotable", Object.ChosenPromotable);
+			Action.Serialize("ChosenLabel", Object.ChosenLabel);
+		}
+	};
+
+	/**
+	 * Converts current object state to JSON value.
+	 *
+	 * @returns JSON value object with current state.
+	 */
+	TSharedRef<FJsonValue> GetStateAsJSON() const
+	{
+		TSharedPtr<FJsonValue> PrevGUIValuePtr;
+		TSharedPtr<FJsonObject> PrevStatePtr;
+
+		if (FSettingsCache::Get().GetSetting(PrevGUIValuePtr, TEXT("GUI"))
+			&& PrevGUIValuePtr->Type == EJson::Object)
+		{
+			PrevStatePtr = PrevGUIValuePtr->AsObject();
+		}
+		else
+		{
+			PrevStatePtr = TSharedPtr<FJsonObject>(new FJsonObject());
+		}
+
+		TSharedRef<FJsonObject> PrevState = PrevStatePtr.ToSharedRef();
+
+		FSettingsCache::SerializeToJSON<FJsonSerializer>(PrevState, *this);
+		
+		return MakeShareable(new FJsonValueObject(PrevStatePtr));
+	}
+
+	/**
+	 * Loads state from JSON value object.
+	 *
+	 * @param Value JSON value object.
+	 */
+	void LoadStateFromJSON(const FJsonValue& State)
+	{
+		if (State.Type == EJson::Object)
+		{
+			FSettingsCache::SerializeFromJSON<FJsonSerializer>(State.AsObject().ToSharedRef(), *this);
+		}
+	}
+
+	/**
+	 * Tells that labels names are currently being loaded.
+	 *
+	 * @returns True if labels names are currently being loaded. False otherwise.
+	 */
+	bool IsLoadingInProgress() const
+	{
+		return LoaderThread.IsValid() && LoaderThread->IsInProgress();
+	}
+
+	/**
+	 * Start async loading of the P4 label data in case user wants it.
+	 */
+	void StartLoadingData()
+	{
+		P4Data->Reset();
+		LoaderThread.Reset();
+
+		OnDataReset.ExecuteIfBound();
+
+		LoaderThread = MakeShareable(new FP4DataLoader(OnDataLoaded, P4Data.Get()));
+	}
+	
+	/**
+	 * Terminates background P4 data loading process.
+	 */
+	void TerminateLoadingProcess()
+	{
+		if (LoaderThread.IsValid())
+		{
+			LoaderThread->Terminate();
+		}
+	}
+
 	/**
 	 * Destructor.
 	 */
 	virtual ~SMainTabWidget()
 	{
-		FUnrealSync::TerminateLoadingProcess();
-		FUnrealSync::TerminateSyncingProcess();
+		TerminateLoadingProcess();
+		TerminateSyncingProcess();
 	}
 
 private:
+	/**
+	 * Gets static shared ptr to the object of this class.
+	 */
+	static TSharedPtr<SMainTabWidget>& GetSingletonPtr()
+	{
+		static TSharedPtr<SMainTabWidget> Ptr = nullptr;
+		return Ptr;
+	}
+
+	/**
+	 * Terminates P4 syncing process.
+	 */
+	void TerminateSyncingProcess()
+	{
+		if (SyncingThread.IsValid())
+		{
+			SyncingThread->Terminate();
+		}
+	}
+
+	/**
+	 * Launches P4 sync command with given command line and options.
+	 *
+	 * @param Settings Sync settings.
+	 * @param LabelNameProvider Object that will provide label name to syncing thread.
+	 * @param OnSyncFinished Delegate to run when syncing is finished.
+	 * @param OnSyncProgress Delegate to run when syncing has made progress.
+	 */
+	void LaunchSync(FSyncSettings Settings, ILabelNameProvider& LabelNameProvider, const FSyncingThread::FOnSyncFinished& OnSyncFinished, const FSyncingThread::FOnSyncProgress& OnSyncProgress)
+	{
+		SyncingThread = MakeShareable(new FSyncingThread(MoveTemp(Settings), LabelNameProvider, OnSyncFinished, OnSyncProgress));
+	}
+
 	/**
 	 * Queues adding lines to the log for execution on Slate rendering thread.
 	 */
@@ -1404,8 +1753,8 @@ private:
 	 */
 	FReply OnCancelButtonClick()
 	{
-		FUnrealSync::TerminateLoadingProcess();
-		FUnrealSync::TerminateSyncingProcess();
+		TerminateLoadingProcess();
+		TerminateSyncingProcess();
 
 		ExternalThreadsDispatcher->AddRenderingThreadRequest(FExternalThreadsDispatcher::FRenderingThreadRequest::CreateRaw(this, &SMainTabWidget::GoBack));
 
@@ -1430,7 +1779,7 @@ private:
 	 */
 	bool IsReloadLabelsReady() const
 	{
-		return !FUnrealSync::IsLoadingInProgress();
+		return !IsLoadingInProgress();
 	}
 
 	/**
@@ -1440,7 +1789,7 @@ private:
 	 */
 	FReply OnReloadLabels()
 	{
-		FUnrealSync::StartLoadingData();
+		StartLoadingData();
 
 		return FReply::Handled();
 	}
@@ -1485,9 +1834,9 @@ private:
 			OverrideSyncStep);
 
 		Switcher->SetActiveWidgetIndex(1);
-		FUnrealSync::LaunchSync(MoveTemp(Settings), GetCurrentSyncCmdLineProvider(),
-			FUnrealSync::FOnSyncFinished::CreateRaw(this, &SMainTabWidget::SyncingFinished),
-			FUnrealSync::FOnSyncProgress::CreateRaw(this, &SMainTabWidget::SyncingProgress));
+		LaunchSync(MoveTemp(Settings), GetCurrentSyncCmdLineProvider(),
+			FSyncingThread::FOnSyncFinished::CreateRaw(this, &SMainTabWidget::SyncingFinished),
+			FSyncingThread::FOnSyncProgress::CreateRaw(this, &SMainTabWidget::SyncingProgress));
 
 		GoBackButton->SetEnabled(false);
 		CancelButton->SetEnabled(true);
@@ -1551,7 +1900,7 @@ private:
 	 */
 	void OnCurrentGameChanged(const FString& CurrentGameName)
 	{
-		if (!FUnrealSync::HasValidData())
+		if (!P4Data->HasValidData())
 		{
 			return;
 		}
@@ -1722,14 +2071,23 @@ private:
 	/* Radio selection used to chose method to sync. */
 	TSharedPtr<SRadioContentSelection> RadioSelection;
 
+	/* Latest promoted widget. */
+	TSharedPtr<SLatestPromoted> LatestPromoted;
+	/* Chosen promoted widget. */
+	TSharedPtr<SPickPromoted> ChosenPromoted;
+	/* Chosen promotable widget. */
+	TSharedPtr<SPickPromotable> ChosenPromotable;
+	/* Chosen label widget. */
+	TSharedPtr<SPickAny> ChosenLabel;
+
 	/* Widget to pick game used to sync. */
 	TSharedPtr<SPickGameWidget> PickGameWidget;
 	/* Check box to tell if this should be an artist sync. */
-	TSharedPtr<SCheckBox> ArtistSyncCheckBox;
+	TSharedPtr<SPreservableCheckBox> ArtistSyncCheckBox;
 	/* Check box to tell if this should be a preview sync. */
-	TSharedPtr<SCheckBox> PreviewSyncCheckBox;
+	TSharedPtr<SPreservableCheckBox> PreviewSyncCheckBox;
 	/* Check box to tell if this should be a auto-clobber sync. */
-	TSharedPtr<SCheckBox> AutoClobberSyncCheckBox;
+	TSharedPtr<SPreservableCheckBox> AutoClobberSyncCheckBox;
 
 	/* External thread requests dispatcher. */
 	TSharedRef<FExternalThreadsDispatcher, ESPMode::ThreadSafe> ExternalThreadsDispatcher;
@@ -1749,6 +2107,21 @@ private:
 
 	/* Override sync step value. */
 	FString OverrideSyncStep;
+
+	/* Data loaded event. */
+	FOnDataLoaded OnDataLoaded;
+
+	/* Data reset event. */
+	FOnDataReset OnDataReset;
+
+	/* Reference to P4 data proxy. */
+	TSharedRef<FP4DataProxy> P4Data;
+
+	/* Background loading process monitoring thread. */
+	TSharedPtr<FP4DataLoader> LoaderThread;
+
+	/* Background syncing process monitoring thread. */
+	TSharedPtr<FSyncingThread> SyncingThread;
 };
 
 void SMainTabWidget::OnOverrideSyncStepText(const FText& CommittedText, ETextCommit::Type Type)
@@ -1855,7 +2228,20 @@ void InitGUI(const TCHAR* CommandLine, bool bP4EnvTabOnly)
 		GLog->FlushThreadedLogs();
 	}
 
+	FUnrealSync::SaveSettingsAndClose();
+
 	FSlateApplication::Shutdown();
 }
 
+void SaveGUISettings()
+{
+	TSharedPtr<SMainTabWidget> Ptr = SMainTabWidget::GetInstance();
+	if (!Ptr.IsValid())
+	{
+		return;
+	}
+
+	FSettingsCache::Get().SetSetting(TEXT("GUI"), Ptr->GetStateAsJSON());
+}
+
 #undef LOCTEXT_NAMESPACE
diff --git a/Engine/Source/Programs/UnrealSync/Private/GUI.h b/Engine/Source/Programs/UnrealSync/Private/GUI.h
--- a/Engine/Source/Programs/UnrealSync/Private/GUI.h
+++ b/Engine/Source/Programs/UnrealSync/Private/GUI.h
@@ -8,4 +8,9 @@
  * @param CommandLine Command line passed to the program.
  * @param bP4EnvTabOnly Tells if initialize GUI with P4 env settings tab only.
  */
-void InitGUI(const TCHAR* CommandLine, bool bP4EnvTabOnly = false);
\ No newline at end of file
+void InitGUI(const TCHAR* CommandLine, bool bP4EnvTabOnly = false);
+
+/**
+ * Saves GUI settings to settings cache.
+ */
+void SaveGUISettings();
diff --git a/Engine/Source/Programs/UnrealSync/Private/P4DataCache.cpp b/Engine/Source/Programs/UnrealSync/Private/P4DataCache.cpp
--- a/Engine/Source/Programs/UnrealSync/Private/P4DataCache.cpp
+++ b/Engine/Source/Programs/UnrealSync/Private/P4DataCache.cpp
@@ -57,8 +57,8 @@ const FDateTime& FP4Label::GetDate() const
 	return Date;
 }
 
-FP4DataLoader::FP4DataLoader(const FOnLoadingFinished& InOnLoadingFinished)
-	: OnLoadingFinished(InOnLoadingFinished), bTerminate(false)
+FP4DataLoader::FP4DataLoader(const FOnLoadingFinished& InOnLoadingFinished, FP4DataProxy& InDataProxy)
+	: OnLoadingFinished(InOnLoadingFinished), DataProxy(InDataProxy), bTerminate(false)
 {
 	Thread = FRunnableThread::Create(this, TEXT("P4 Data Loading"));
 }
@@ -92,16 +92,17 @@ uint32 FP4DataLoader::Run()
 	if (!FP4Env::RunP4Progress(FString::Printf(TEXT("labels -t -e%s/*"), *FP4Env::Get().GetBranch()),
 		FP4Env::FOnP4MadeProgress::CreateRaw(&Progress, &P4Progress::OnProgress)))
 	{
-		OnLoadingFinished.ExecuteIfBound(nullptr);
+		OnLoadingFinished.ExecuteIfBound();
 		return 0;
 	}
 
 	if (!Data->LoadFromLog(Progress.Output))
 	{
-		OnLoadingFinished.ExecuteIfBound(nullptr);
+		OnLoadingFinished.ExecuteIfBound();
 	}
 
-	OnLoadingFinished.ExecuteIfBound(Data);
+	DataProxy.OnP4DataLoadingFinished(Data);
+	OnLoadingFinished.ExecuteIfBound();
 
 	return 0;
 }
@@ -133,3 +134,171 @@ bool FP4DataLoader::Init()
 
 	return true;
 }
+
+FString FP4DataProxy::GetLatestLabelForGame(const FString& GameName)
+{
+	auto Labels = GetPromotedLabelsForGame(GameName);
+	return Labels->Num() == 0 ? "" : (*Labels)[0];
+}
+
+/**
+ * Tells if label is in the format: <BranchPath>/<Prefix>-<GameNameIfNotEmpty>-CL-*
+ *
+ * @param LabelName Label name to check.
+ * @param LabelNamePrefix Prefix to use.
+ * @param GameName Game name to use.
+ *
+ * @returns True if label is in the format. False otherwise.
+ */
+bool IsPrefixedGameLabelName(const FString& LabelName, const FString& LabelNamePrefix, const FString& GameName)
+{
+	int32 BranchNameEnd = 0;
+
+	if (!LabelName.FindLastChar('/', BranchNameEnd))
+	{
+		return false;
+	}
+
+	FString Rest = LabelName.Mid(BranchNameEnd + 1);
+
+	FString LabelNameGamePart;
+	if (!GameName.Equals(""))
+	{
+		LabelNameGamePart = "-" + GameName;
+	}
+
+	return Rest.StartsWith(LabelNamePrefix + LabelNameGamePart + "-CL-", ESearchCase::CaseSensitive);
+}
+
+/**
+ * Tells if label is in the format: <BranchPath>/Promoted-<GameNameIfNotEmpty>-CL-*
+ *
+ * @param LabelName Label name to check.
+ * @param GameName Game name to use.
+ *
+ * @returns True if label is in the format. False otherwise.
+ */
+bool IsPromotedGameLabelName(const FString& LabelName, const FString& GameName)
+{
+	return IsPrefixedGameLabelName(LabelName, "Promoted", GameName);
+}
+
+/**
+ * Tells if label is in the format: <BranchPath>/Promotable-<GameNameIfNotEmpty>-CL-*
+ *
+ * @param LabelName Label name to check.
+ * @param GameName Game name to use.
+ *
+ * @returns True if label is in the format. False otherwise.
+ */
+bool IsPromotableGameLabelName(const FString& LabelName, const FString& GameName)
+{
+	return IsPrefixedGameLabelName(LabelName, "Promotable", GameName);
+}
+
+TSharedPtr<TArray<FString> > FP4DataProxy::GetPromotedLabelsForGame(const FString& GameName)
+{
+	struct TFillLabelsPolicy 
+	{
+		static void Fill(TArray<FString>& OutLabelNames, const FString& InGameName, const TArray<FP4Label>& InLabels)
+		{
+			for (auto Label : InLabels)
+			{
+				if (IsPromotedGameLabelName(Label.GetName(), InGameName))
+				{
+					OutLabelNames.Add(Label.GetName());
+				}
+			}
+		}
+	};
+
+	return GetLabelsForGame<TFillLabelsPolicy>(GameName);
+}
+
+TSharedPtr<TArray<FString> > FP4DataProxy::GetPromotableLabelsForGame(const FString& GameName)
+{
+	struct TFillLabelsPolicy
+	{
+		static void Fill(TArray<FString>& OutLabelNames, const FString& InGameName, const TArray<FP4Label>& InLabels)
+		{
+			for (auto Label : InLabels)
+			{
+				if (IsPromotedGameLabelName(Label.GetName(), InGameName))
+				{
+					break;
+				}
+
+				if (IsPromotableGameLabelName(Label.GetName(), InGameName))
+				{
+					OutLabelNames.Add(Label.GetName());
+				}
+			}
+		}
+	};
+
+	return GetLabelsForGame<TFillLabelsPolicy>(GameName);
+}
+
+TSharedPtr<TArray<FString> > FP4DataProxy::GetAllLabels()
+{
+	struct TFillLabelsPolicy
+	{
+		static void Fill(TArray<FString>& OutLabelNames, const FString& InGameName, const TArray<FP4Label>& InLabels)
+		{
+			for (auto Label : InLabels)
+			{
+				OutLabelNames.Add(Label.GetName());
+			}
+		}
+	};
+
+	return GetLabelsForGame<TFillLabelsPolicy>("");
+}
+
+TSharedPtr<TArray<FString> > FP4DataProxy::GetPossibleGameNames()
+{
+	TSharedPtr<TArray<FString> > PossibleGames = MakeShareable(new TArray<FString>());
+
+	FP4Env& Env = FP4Env::Get();
+
+	FString FileList;
+	if (!Env.RunP4Output(FString::Printf(TEXT("files -e %s/.../Build/ArtistSyncRules.xml"), *Env.GetBranch()), FileList) || FileList.IsEmpty())
+	{
+		return PossibleGames;
+	}
+
+	FString Line, Rest = FileList;
+	while (Rest.Split(LINE_TERMINATOR, &Line, &Rest, ESearchCase::CaseSensitive))
+	{
+		if (!Line.StartsWith(Env.GetBranch()))
+		{
+			continue;
+		}
+
+		int32 ArtistSyncRulesPos = Line.Find("/Build/ArtistSyncRules.xml#", ESearchCase::IgnoreCase);
+
+		if (ArtistSyncRulesPos == INDEX_NONE)
+		{
+			continue;
+		}
+
+		FString MiddlePart = Line.Mid(Env.GetBranch().Len(), ArtistSyncRulesPos - Env.GetBranch().Len());
+
+		int32 LastSlash = INDEX_NONE;
+		MiddlePart.FindLastChar('/', LastSlash);
+
+		PossibleGames->Add(MiddlePart.Mid(LastSlash + 1));
+	}
+
+	return PossibleGames;
+}
+
+void FP4DataProxy::OnP4DataLoadingFinished(TSharedPtr<FP4DataCache> InData)
+{
+	Data = InData;
+}
+
+bool FP4DataProxy::HasValidData() const
+{
+	return Data.IsValid();
+}
\ No newline at end of file
diff --git a/Engine/Source/Programs/UnrealSync/Private/P4DataCache.h b/Engine/Source/Programs/UnrealSync/Private/P4DataCache.h
--- a/Engine/Source/Programs/UnrealSync/Private/P4DataCache.h
+++ b/Engine/Source/Programs/UnrealSync/Private/P4DataCache.h
@@ -63,6 +63,8 @@ private:
 	TArray<FP4Label> Labels;
 };
 
+class FP4DataProxy;
+
 /**
  * Class to store info about thread running and controling P4 data loading process.
  */
@@ -70,14 +72,14 @@ class FP4DataLoader : public FRunnable
 {
 public:
 	/* Delegate for loading finished event. */
-	DECLARE_DELEGATE_OneParam(FOnLoadingFinished, TSharedPtr<FP4DataCache>);
+	DECLARE_DELEGATE(FOnLoadingFinished);
 
 	/**
 	 * Constructor
 	 *
 	 * @param InOnLoadingFinished Delegate to run when loading finished.
 	 */
-	FP4DataLoader(const FOnLoadingFinished& InOnLoadingFinished);
+	FP4DataLoader(const FOnLoadingFinished& InOnLoadingFinished, FP4DataProxy& InDataProxy);
 
 	virtual ~FP4DataLoader();
 
@@ -124,6 +126,9 @@ private:
 	/* Handle for thread object. */
 	FRunnableThread *Thread;
 
+	/* Reference to data proxy. */
+	FP4DataProxy& DataProxy;
+
 	/* If this thread is still running. */
 	bool bInProgress;
 
@@ -131,3 +136,96 @@ private:
 	bool bTerminate;
 };
 
+class FP4DataProxy
+{
+public:
+	/**
+	 * Gets latest label for given game name.
+	 *
+	 * @param GameName Current game name.
+	 *
+	 * @returns Latest label for current game.
+	 */
+	FString GetLatestLabelForGame(const FString& GameName);
+
+	/**
+	 * Gets promoted labels for given game.
+	 *
+	 * @param GameName Current game name.
+	 *
+	 * @returns Array of promoted labels for given game.
+	 */
+	TSharedPtr<TArray<FString> > GetPromotedLabelsForGame(const FString& GameName);
+
+	/**
+	 * Gets promotable labels for given game since last promoted.
+	 *
+	 * @param GameName Current game name.
+	 *
+	 * @returns Array of promotable labels for given game since last promoted.
+	 */
+	TSharedPtr<TArray<FString> > GetPromotableLabelsForGame(const FString& GameName);
+
+	/**
+	 * Gets all labels.
+	 *
+	 * @returns Array of all labels names.
+	 */
+	TSharedPtr<TArray<FString> > GetAllLabels();
+
+	/**
+	 * Gets possible game names.
+	 *
+	 * @returns Array of possible game names.
+	 */
+	TSharedPtr<TArray<FString> > GetPossibleGameNames();
+
+	/**
+	 * Tells if has valid P4 data loaded.
+	 *
+	 * @returns If P4 data was loaded.
+	 */
+	bool HasValidData() const;
+
+	/**
+	 * Clears data in proxy.
+	 */
+	void Reset()
+	{
+		Data.Reset();
+	}
+
+	/**
+	 * Method to receive p4 data loading finished event.
+	 *
+	 * @param InData Loaded data.
+	 */
+	void OnP4DataLoadingFinished(TSharedPtr<FP4DataCache> InData);
+
+private:
+	/**
+	 * Helper function to construct and get array of labels if data is fetched from the P4
+	 * using given policy.
+	 *
+	 * @template_param TFillLabelsPolicy Policy class that has to implement:
+	 *     static void Fill(TArray<FString>& OutLabelNames, const FString& GameName, const TArray<FP4Label>& Labels)
+	 *     that will fill the array on request.
+	 * @param GameName Game name to pass to the Fill method from policy.
+	 * @returns Filled out labels.
+	 */
+	template <class TFillLabelsPolicy>
+	TSharedPtr<TArray<FString> > GetLabelsForGame(const FString& GameName)
+	{
+		TSharedPtr<TArray<FString> > OutputLabels = MakeShareable(new TArray<FString>());
+
+		if (Data.IsValid())
+		{
+			TFillLabelsPolicy::Fill(*OutputLabels, GameName, Data->GetLabels());
+		}
+
+		return OutputLabels;
+	}
+
+	/* Cached data ptr. */
+	TSharedPtr<FP4DataCache> Data;
+};
\ No newline at end of file
diff --git a/Engine/Source/Programs/UnrealSync/Private/P4Env.cpp b/Engine/Source/Programs/UnrealSync/Private/P4Env.cpp
--- a/Engine/Source/Programs/UnrealSync/Private/P4Env.cpp
+++ b/Engine/Source/Programs/UnrealSync/Private/P4Env.cpp
@@ -5,12 +5,9 @@
 #include "P4Env.h"
 #include "ProcessHelper.h"
 #include "Regex.h"
-#include "Map.h"
-#include "JsonSerializer.h"
-#include "JsonObject.h"
-#include "JsonReader.h"
 #include "SEditableTextBox.h"
 #include "SGridPanel.h"
+#include "SettingsCache.h"
 
 #define LOCTEXT_NAMESPACE "UnrealSync"
 
@@ -30,151 +27,6 @@ bool GetCommandLineParam(FString& Value, const TCHAR* CommandLine, EP4ParamType
 	return FParse::Value(CommandLine, *(FP4Env::GetParamName(Type) + "="), Value);
 }
 
-/**
- * Cache class for file settings.
- */
-class FSettingsCache
-{
-public:
-	/**
-	 * Instance getter.
-	 */
-	static FSettingsCache& Get()
-	{
-		static FSettingsCache Cache;
-		return Cache;
-	}
-
-	/**
-	 * Gets setting of given type.
-	 *
-	 * @param Value Output parameter. If succeeded this parameter will be overwritten with setting value.
-	 * @param Type Type to look for.
-	 *
-	 * @returns True if setting was found. False otherwise.
-	 */
-	bool GetSetting(FString& Value, EP4ParamType Type)
-	{
-		auto* ParamPtr = Settings.Find(FP4Env::GetParamName(Type));
-		if (ParamPtr == nullptr)
-		{
-			return false;
-		}
-
-		Value = *ParamPtr;
-		return true;
-	}
-
-	/**
-	 * Sets setting of given type to given value.
-	 *
-	 * @param Type Type of the setting.
-	 * @param Value Value of the setting.
-	 */
-	void SetSetting(EP4ParamType Type, const FString& Value)
-	{
-		auto* ParamPtr = Settings.Find(FP4Env::GetParamName(Type));
-		if (ParamPtr != nullptr)
-		{
-			*ParamPtr = Value;
-		}
-
-		Settings.Add(FP4Env::GetParamName(Type), Value);
-	}
-
-	/**
-	 * Saves setting to file.
-	 */
-	void Save()
-	{
-		TSharedRef<FJsonObject> Object(new FJsonObject());
-
-		for (const auto& Setting : Settings)
-		{
-			if (!Setting.Value.IsEmpty())
-			{
-				Object->SetStringField(Setting.Key, Setting.Value);
-			}
-		}
-
-		FString Buffer;
-		TSharedRef<TJsonWriter<> > Writer = TJsonWriterFactory<>::Create(&Buffer);
-		if (!FJsonSerializer::Serialize(Object, Writer))
-		{
-			UE_LOG(LogP4Env, Error, TEXT("Failed serializing settings using JSON."));
-			return;
-		}
-
-		auto bNeedsWrite = true;
-
-		if (FPaths::FileExists(GetSettingsFileName()))
-		{
-			// Read the file to a string.
-			FString FileContents;
-			bNeedsWrite = !(FFileHelper::LoadFileToString(FileContents, *GetSettingsFileName()) && FileContents == Buffer);
-		}
-
-		if (bNeedsWrite && !FFileHelper::SaveStringToFile(Buffer, *GetSettingsFileName()))
-		{
-			UE_LOG(LogP4Env, Error, TEXT("Failed writing settings to file: \"%s\"."), *GetSettingsFileName());
-			return;
-		}
-	}
-
-private:
-	/**
-	 * Constructor.
-	 *
-	 * Creates the cache from value read from settings file if it exists.
-	 */
-	FSettingsCache()
-	{
-		// Read the file to a string.
-		FString FileContents;
-		if (!FFileHelper::LoadFileToString(FileContents, *GetSettingsFileName()))
-		{
-			UE_LOG(LogP4Env, Log, TEXT("Couldn't find settings file \"%s\"."), *GetSettingsFileName());
-			return;
-		}
-
-		// Deserialize a JSON object from the string
-		TSharedPtr<FJsonObject> Object;
-		TSharedRef<TJsonReader<> > Reader = TJsonReaderFactory<>::Create(FileContents);
-		if (!FJsonSerializer::Deserialize(Reader, Object) || !Object.IsValid())
-		{
-			UE_LOG(LogP4Env, Error, TEXT("Settings file JSON parsing failed: \"%s\"."), *Reader->GetErrorMessage());
-			return;
-		}
-
-		for (const auto& Value : Object->Values)
-		{
-			FString Text;
-			if (!Value.Value->TryGetString(Text))
-			{
-				Settings.Empty();
-				break;
-			}
-
-			Settings.Add(Value.Key, Text);
-		}
-	}
-
-	/**
-	 * Gets settings file name. For now it's hardcoded.
-	 *
-	 * @returns Settings file name.
-	 */
-	static const FString& GetSettingsFileName()
-	{
-		const static FString HardCodedSettingsFileName = "UnrealSync.settings";
-
-		return HardCodedSettingsFileName;
-	}
-
-	/** Value map of settings. */
-	TMap<FString, FString> Settings;
-};
-
 /**
  * Gets param from environment variables.
  *
@@ -300,7 +152,7 @@ public:
 			OverriddenSetting(Type, HardParam, true);
 		}
 
-		if (FSettingsCache::Get().GetSetting(HardParam, Type))
+		if (FP4Env::GetSetting(HardParam, Type))
 		{
 			OverriddenSetting(Type, HardParam, false);
 		}
@@ -1108,6 +960,18 @@ bool FP4Env::CheckIfFileNeedsUpdate(const FString& FilePath)
 	return HaveRev < HeadRev;
 }
 
+bool FP4Env::GetSetting(FString& Value, EP4ParamType Type)
+{
+	TSharedPtr<FJsonValue> JSONValue;
+	if (!FSettingsCache::Get().GetSetting(JSONValue, TEXT("Perforce"))
+		|| JSONValue->Type != EJson::Object)
+	{
+		return false;
+	}
+
+	return JSONValue->AsObject()->TryGetStringField(GetParamName(Type), Value);
+}
+
 bool FP4Env::IsValid()
 {
 	return Env.IsValid();
@@ -1409,16 +1273,19 @@ END_SLATE_FUNCTION_BUILD_OPTIMIZATION
 FReply SP4EnvTabWidget::OnSaveAndRestartButtonClick()
 {
 	auto& Settings = FSettingsCache::Get();
+	TSharedPtr<FJsonObject> Object(new FJsonObject());
 
 	for (auto Option : Options)
 	{
-		Settings.SetSetting(Option->GetType(), Option->GetText().ToString());
+		if (!Option->GetText().IsEmpty())
+		{
+			Object->SetStringField(FP4Env::GetParamName(Option->GetType()), Option->GetText().ToString());
+		}
 	}
 
-	Settings.Save();
+	Settings.SetSetting(TEXT("Perforce"), MakeShareable(new FJsonValueObject(Object)));
 
-	FUnrealSync::RunDetachedUS(FPlatformProcess::ExecutableName(false), true, true, false);
-	FPlatformMisc::RequestExit(false);
+	FUnrealSync::SaveSettingsAndRestart();
 
 	return FReply::Handled();
 }
@@ -1434,7 +1301,7 @@ SP4EnvTabWidget::FP4Option::FP4Option(EP4ParamType InType)
 	: Type(InType)
 {
 	FString Param;
-	if (FSettingsCache::Get().GetSetting(Param, Type))
+	if (FP4Env::Get().GetSetting(Param, Type))
 	{
 		Text = FText::FromString(Param);
 	}
diff --git a/Engine/Source/Programs/UnrealSync/Private/P4Env.h b/Engine/Source/Programs/UnrealSync/Private/P4Env.h
--- a/Engine/Source/Programs/UnrealSync/Private/P4Env.h
+++ b/Engine/Source/Programs/UnrealSync/Private/P4Env.h
@@ -234,6 +234,16 @@ public:
 	 */
 	static bool CheckIfFileNeedsUpdate(const FString& FilePath);
 
+	/**
+	 * Gets setting of given type.
+	 *
+	 * @param Value Output parameter. If succeeded this parameter will be overwritten with setting value.
+	 * @param Type Type to look for.
+	 *
+	 * @returns True if setting was found. False otherwise.
+	 */
+	static bool GetSetting(FString& Value, EP4ParamType Type);
+
 private:
 	/* Param serialization delegate. */
 	DECLARE_DELEGATE_TwoParams(FSerializationTask, FString&, EP4ParamType);
@@ -284,4 +294,7 @@ private:
 
 	/* P4 current branch. */
 	FString Branch;
+
+	/** Value map of P4 settings. */
+	TMap<FString, FString> P4FileSettingsCache;
 };
diff --git a/Engine/Source/Programs/UnrealSync/Private/SettingsCache.cpp b/Engine/Source/Programs/UnrealSync/Private/SettingsCache.cpp
new file mode 100644
--- /dev/null
+++ b/Engine/Source/Programs/UnrealSync/Private/SettingsCache.cpp
@@ -0,0 +1,111 @@
+// Copyright 1998-2015 Epic Games, Inc. All Rights Reserved.
+
+#include "UnrealSync.h"
+#include "SettingsCache.h"
+
+DECLARE_LOG_CATEGORY_CLASS(LogSettings, Log, All);
+
+const FString& FSettingsCache::GetSettingsFileName()
+{
+	const static FString HardCodedSettingsFileName = "UnrealSync.settings";
+
+	return HardCodedSettingsFileName;
+}
+
+FSettingsCache::FSettingsCache()
+{
+	// Read the file to a string.
+	FString FileContents;
+	if (!FFileHelper::LoadFileToString(FileContents, *GetSettingsFileName()))
+	{
+		UE_LOG(LogSettings, Log, TEXT("Couldn't find settings file \"%s\"."), *GetSettingsFileName());
+		return;
+	}
+
+	// Deserialize a JSON object from the string
+	TSharedPtr<FJsonObject> Object;
+	TSharedRef<TJsonReader<> > Reader = TJsonReaderFactory<>::Create(FileContents);
+	if (!FJsonSerializer::Deserialize(Reader, Object) || !Object.IsValid())
+	{
+		UE_LOG(LogSettings, Error, TEXT("Settings file JSON parsing failed: \"%s\"."), *Reader->GetErrorMessage());
+		return;
+	}
+
+	TMap<FString, TSharedRef<FJsonValue> > TempSettings;
+	for (const auto& KeyValuePair : Object->Values)
+	{
+		if (!KeyValuePair.Value.IsValid())
+		{
+			UE_LOG(LogSettings, Error, TEXT("Bad format of settings file: \"%s\"."), *Reader->GetErrorMessage());
+			return;
+		}
+
+		TempSettings.Add(KeyValuePair.Key, KeyValuePair.Value.ToSharedRef());
+	}
+
+	Settings = MoveTemp(TempSettings);
+}
+
+void FSettingsCache::Save()
+{
+	TSharedRef<FJsonObject> Object(new FJsonObject());
+
+	for (const auto& Setting : Settings)
+	{
+		Object->SetField(Setting.Key, Setting.Value);
+	}
+
+	FString Buffer;
+	TSharedRef<TJsonWriter<> > Writer = TJsonWriterFactory<>::Create(&Buffer);
+	if (!FJsonSerializer::Serialize(Object, Writer))
+	{
+		UE_LOG(LogSettings, Error, TEXT("Failed serializing settings using JSON."));
+		return;
+	}
+
+	auto bNeedsWrite = true;
+
+	if (FPaths::FileExists(GetSettingsFileName()))
+	{
+		// Read the file to a string.
+		FString FileContents;
+		bNeedsWrite = !(FFileHelper::LoadFileToString(FileContents, *GetSettingsFileName()) && FileContents == Buffer);
+	}
+
+	if (bNeedsWrite && !FFileHelper::SaveStringToFile(Buffer, *GetSettingsFileName()))
+	{
+		UE_LOG(LogSettings, Error, TEXT("Failed writing settings to file: \"%s\"."), *GetSettingsFileName());
+		return;
+	}
+}
+
+void FSettingsCache::SetSetting(const FString& Name, TSharedRef<FJsonValue> Value)
+{
+	auto* ParamPtr = Settings.Find(Name);
+	if (ParamPtr != nullptr)
+	{
+		*ParamPtr = Value;
+	}
+	else
+	{
+		Settings.Add(Name, Value);
+	}
+}
+
+bool FSettingsCache::GetSetting(TSharedPtr<FJsonValue>& Value, const FString& Name)
+{
+	auto* ParamPtr = Settings.Find(Name);
+	if (ParamPtr == nullptr)
+	{
+		return false;
+	}
+
+	Value = *ParamPtr;
+	return true;
+}
+
+FSettingsCache& FSettingsCache::Get()
+{
+	static FSettingsCache Cache;
+	return Cache;
+}
diff --git a/Engine/Source/Programs/UnrealSync/Private/SettingsCache.h b/Engine/Source/Programs/UnrealSync/Private/SettingsCache.h
new file mode 100644
--- /dev/null
+++ b/Engine/Source/Programs/UnrealSync/Private/SettingsCache.h
@@ -0,0 +1,261 @@
+// Copyright 1998-2015 Epic Games, Inc. All Rights Reserved.
+
+#pragma once
+
+#include "Map.h"
+#include "JsonSerializer.h"
+#include "JsonObject.h"
+#include "JsonReader.h"
+
+/**
+ * Cache class for file settings.
+ */
+class FSettingsCache
+{
+public:
+	/**
+	 * Instance getter.
+	 */
+	static FSettingsCache& Get();
+
+	/**
+	 * Gets setting of given type.
+	 *
+	 * @param Value Output parameter. If succeeded this parameter will be overwritten with setting JSON value.
+	 * @param Name Name to look for.
+	 *
+	 * @returns True if setting was found. False otherwise.
+	 */
+	bool GetSetting(TSharedPtr<FJsonValue>& Value, const FString& Name);
+
+	/**
+	 * Sets setting of given type to given value.
+	 *
+	 * @param Name Name of the setting.
+	 * @param Value Value of the setting.
+	 */
+	void SetSetting(const FString& Name, TSharedRef<FJsonValue> Value);
+
+	/**
+	 * Saves setting to file.
+	 */
+	void Save();
+
+	/**
+	 * TAddConstIf trait to conditionally add const to reference type.
+	 */
+	template <bool bCondition, typename Type>
+	struct TAddConstIf;
+
+	template <bool bCondition, typename Type>
+	struct TAddConstIf<bCondition, Type&>
+	{
+		typedef typename TChooseClass<bCondition, const Type&, Type&>::Result Result;
+	};
+
+	/**
+	 * Pointer traits to strip TSharedRef or TSharedPtr and return simple
+	 * object reference (T&) to unify calls.
+	 */
+	template <typename T, bool bIsConst>
+	struct TPointerTraits
+	{
+		typedef typename TAddConstIf<bIsConst, T&>::Result ObjectRefType;
+		typedef typename TAddConstIf<bIsConst, T&>::Result ParamType;
+
+		static ObjectRefType GetObject(ParamType Object) { return Object; }
+	};
+
+	template <typename T, bool bIsConst>
+	struct TPointerTraits<TSharedRef<T>, bIsConst>
+	{
+		typedef typename TAddConstIf<bIsConst, T&>::Result ObjectRefType;
+		typedef typename TAddConstIf<bIsConst, TSharedRef<T>&>::Result ParamType;
+
+		static ObjectRefType GetObject(ParamType Object) { return Object.Get(); }
+	};
+
+	template <typename T, bool bIsConst>
+	struct TPointerTraits<TSharedPtr<T>, bIsConst>
+	{
+		typedef typename TAddConstIf<bIsConst, T&>::Result ObjectRefType;
+		typedef typename TAddConstIf<bIsConst, TSharedPtr<T>&>::Result ParamType;
+
+		static ObjectRefType GetObject(ParamType Object) { return *Object.Get(); }
+	};
+
+	/**
+	 * Forwards GetStateAsJSON call to preservable object.
+	 *
+	 * @param Object Object to forward call for.
+	 * @returns Returned JSON value from forwarded call.
+	 */
+	template <typename TPreservableType>
+	static TSharedPtr<FJsonValue> GetStateAsJSON(const TPreservableType Object)
+	{
+		typedef TPointerTraits<TPreservableType, true> T;
+		T::ObjectRefType Ref = T::GetObject(Object);
+
+		return Ref.GetStateAsJSON();
+	}
+
+	/**
+	 * Forwards LoadStateFromJSON call to preservable object.
+	 *
+	 * @param Object Object to forward call for.
+	 * @param Value JSON vale to forward.
+	 */
+	template <typename TPreservableType, typename TValueType>
+	static void LoadStateFromJSON(TPreservableType Object, const TValueType Value)
+	{
+		typedef TPointerTraits<TPreservableType, false> PT;
+		PT::ObjectRefType Ref = PT::GetObject(Object);
+
+		typedef TPointerTraits<TValueType, true> VT;
+		VT::ObjectRefType ValRef = VT::GetObject(Value);
+
+		Ref.LoadStateFromJSON(ValRef);
+	}
+
+	/**
+	 * Saves single object to JSON.
+	 *
+	 * @param JSONObject JSON object to save data into.
+	 * @param FieldName JSON field name for object's data.
+	 * @param Object Object to save.
+	 */
+	template <typename TPreservableType>
+	static void SaveObjectToJSON(TSharedRef<FJsonObject> JSONObject, const FString& FieldName, const TPreservableType Object)
+	{
+		TSharedPtr<FJsonValue> Value = FSettingsCache::GetStateAsJSON(Object);
+
+		if (Value.IsValid())
+		{
+			JSONObject->SetField(FieldName, Value);
+		}
+	}
+
+	/**
+	 * Loads single object from JSON.
+	 *
+	 * @param JSONObject JSON object that's containing data to load.
+	 * @param FieldName JSON field name for object's data.
+	 * @param Object Object to load.
+	 */
+	template <typename TPreservableType>
+	static void LoadObjectFromJSON(const TSharedRef<FJsonObject> JSONObject, const FString& FieldName, TPreservableType Object)
+	{
+		TSharedPtr<FJsonValue> Field = JSONObject->TryGetField(FieldName);
+
+		if (Field.IsValid())
+		{
+			LoadStateFromJSON(Object, Field);
+		}
+	}
+
+	/**
+	 * Save objects serialization action policy.
+	 */
+	struct FSaveObjectsSerializationActionPolicy
+	{
+	public:
+		FSaveObjectsSerializationActionPolicy(TSharedRef<FJsonObject> InJSON)
+			: JSON(MoveTemp(InJSON))
+		{}
+
+		/**
+		 * Tells what to do on serialization action. It just saves serialized
+		 * object to JSON.
+		 *
+		 * @param FieldName JSON field name.
+		 * @param Object Object to save.
+		 */
+		template <class TPreservableType>
+		void Serialize(const FString& FieldName, const TPreservableType Object) const
+		{
+			FSettingsCache::SaveObjectToJSON(JSON, FieldName, Object);
+		}
+
+	private:
+		/* Stored JSON object for serialization. */
+		TSharedRef<FJsonObject> JSON;
+	};
+
+	/**
+	 * Saves object to JSON using TSerializationPolicy.
+	 *
+	 * @template_param TSerializationPolicy Policy that tells the method how to
+	 *		serialize Object. This method uses const version, marked by first
+	 *		template parameter equals to true.
+	 * @param JSON Destination JSON object.
+	 * @param Object Source object.
+	 */
+	template <template <bool> class TSerializationPolicy, typename TObjectToSerializeType>
+	static void SerializeToJSON(TSharedRef<FJsonObject> JSON, const TObjectToSerializeType& Object)
+	{
+		FSaveObjectsSerializationActionPolicy Action(JSON);
+		TSerializationPolicy<true>::Serialize(Action, Object);
+	}
+
+	/**
+	 * Load objects serialization action policy.
+	 */
+	struct FLoadObjectsSerializationActionPolicy
+	{
+	public:
+		FLoadObjectsSerializationActionPolicy(TSharedRef<FJsonObject> InJSON)
+			: JSON(MoveTemp(InJSON))
+		{}
+
+		/**
+		 * Tells what to do on serialization action. It just loads serialized
+		 * object from JSON.
+		 *
+		 * @param FieldName JSON field name.
+		 * @param Object Object to load.
+		 */
+		template <class TPreservableType>
+		void Serialize(const FString& FieldName, TPreservableType Object)
+		{
+			FSettingsCache::LoadObjectFromJSON(JSON, FieldName, Object);
+		}
+
+	private:
+		/* Stored JSON object for serialization. */
+		TSharedRef<FJsonObject> JSON;
+	};
+
+	/**
+	 * Loads object from JSON using TSerializationPolicy.
+	 *
+	 * @template_param TSerializationPolicy Policy that tells the method how to
+	 *		serialize Object. This method uses non-const version, marked by
+	 *		first template parameter equals to false.
+	 * @param JSON Data source.
+	 * @param Object Destination object.
+	 */
+	template <template <bool> class TSerializationPolicy, typename TObjectToSerializeType>
+	static void SerializeFromJSON(TSharedRef<FJsonObject> JSON, TObjectToSerializeType& Object)
+	{
+		FLoadObjectsSerializationActionPolicy Action(JSON);
+		TSerializationPolicy<false>::Serialize(Action, Object);
+	}
+
+private:
+	/**
+	 * Constructor.
+	 *
+	 * Creates the cache from value read from settings file if it exists.
+	 */
+	FSettingsCache();
+
+	/**
+	 * Gets settings file name. For now it's hardcoded.
+	 *
+	 * @returns Settings file name.
+	 */
+	static const FString& GetSettingsFileName();
+
+	/** Map of settings' JSON values. */
+	TMap<FString, TSharedRef<FJsonValue> > Settings;
+};
diff --git a/Engine/Source/Programs/UnrealSync/Private/UnrealSync.cpp b/Engine/Source/Programs/UnrealSync/Private/UnrealSync.cpp
--- a/Engine/Source/Programs/UnrealSync/Private/UnrealSync.cpp
+++ b/Engine/Source/Programs/UnrealSync/Private/UnrealSync.cpp
@@ -5,6 +5,7 @@
 #include "P4DataCache.h"
 #include "P4Env.h"
 #include "XmlParser.h"
+#include "SettingsCache.h"
 
 #include "RequiredProgramMainCPPInclude.h"
 
@@ -83,187 +84,6 @@ void FUnrealSync::RunUnrealSync(const TCHAR* CommandLine)
 	}
 }
 
-FString FUnrealSync::GetLatestLabelForGame(const FString& GameName)
-{
-	auto Labels = GetPromotedLabelsForGame(GameName);
-	return Labels->Num() == 0 ? "" : (*Labels)[0];
-}
-
-/**
- * Tells if label is in the format: <BranchPath>/<Prefix>-<GameNameIfNotEmpty>-CL-*
- *
- * @param LabelName Label name to check.
- * @param LabelNamePrefix Prefix to use.
- * @param GameName Game name to use.
- *
- * @returns True if label is in the format. False otherwise.
- */
-bool IsPrefixedGameLabelName(const FString& LabelName, const FString& LabelNamePrefix, const FString& GameName)
-{
-	int32 BranchNameEnd = 0;
-
-	if (!LabelName.FindLastChar('/', BranchNameEnd))
-	{
-		return false;
-	}
-
-	FString Rest = LabelName.Mid(BranchNameEnd + 1);
-
-	FString LabelNameGamePart;
-	if (!GameName.Equals(""))
-	{
-		LabelNameGamePart = "-" + GameName;
-	}
-
-	return Rest.StartsWith(LabelNamePrefix + LabelNameGamePart + "-CL-", ESearchCase::CaseSensitive);
-}
-
-/**
- * Tells if label is in the format: <BranchPath>/Promoted-<GameNameIfNotEmpty>-CL-*
- *
- * @param LabelName Label name to check.
- * @param GameName Game name to use.
- *
- * @returns True if label is in the format. False otherwise.
- */
-bool IsPromotedGameLabelName(const FString& LabelName, const FString& GameName)
-{
-	return IsPrefixedGameLabelName(LabelName, "Promoted", GameName);
-}
-
-/**
- * Tells if label is in the format: <BranchPath>/Promotable-<GameNameIfNotEmpty>-CL-*
- *
- * @param LabelName Label name to check.
- * @param GameName Game name to use.
- *
- * @returns True if label is in the format. False otherwise.
- */
-bool IsPromotableGameLabelName(const FString& LabelName, const FString& GameName)
-{
-	return IsPrefixedGameLabelName(LabelName, "Promotable", GameName);
-}
-
-/**
- * Helper function to construct and get array of labels if data is fetched from the P4
- * using given policy.
- *
- * @template_param TFillLabelsPolicy Policy class that has to implement:
- *     static void Fill(TArray<FString>& OutLabelNames, const FString& GameName, const TArray<FP4Label>& Labels)
- *     that will fill the array on request.
- * @param GameName Game name to pass to the Fill method from policy.
- * @returns Filled out labels.
- */
-template <class TFillLabelsPolicy>
-TSharedPtr<TArray<FString> > GetLabelsForGame(const FString& GameName)
-{
-	TSharedPtr<TArray<FString> > OutputLabels = MakeShareable(new TArray<FString>());
-
-	if (FUnrealSync::HasValidData())
-	{
-		TFillLabelsPolicy::Fill(*OutputLabels, GameName, FUnrealSync::GetLabels());
-	}
-
-	return OutputLabels;
-}
-
-TSharedPtr<TArray<FString> > FUnrealSync::GetPromotedLabelsForGame(const FString& GameName)
-{
-	struct TFillLabelsPolicy 
-	{
-		static void Fill(TArray<FString>& OutLabelNames, const FString& InGameName, const TArray<FP4Label>& InLabels)
-		{
-			for (auto Label : InLabels)
-			{
-				if (IsPromotedGameLabelName(Label.GetName(), InGameName))
-				{
-					OutLabelNames.Add(Label.GetName());
-				}
-			}
-		}
-	};
-
-	return GetLabelsForGame<TFillLabelsPolicy>(GameName);
-}
-
-TSharedPtr<TArray<FString> > FUnrealSync::GetPromotableLabelsForGame(const FString& GameName)
-{
-	struct TFillLabelsPolicy
-	{
-		static void Fill(TArray<FString>& OutLabelNames, const FString& InGameName, const TArray<FP4Label>& InLabels)
-		{
-			for (auto Label : InLabels)
-			{
-				if (IsPromotedGameLabelName(Label.GetName(), InGameName))
-				{
-					break;
-				}
-
-				if (IsPromotableGameLabelName(Label.GetName(), InGameName))
-				{
-					OutLabelNames.Add(Label.GetName());
-				}
-			}
-		}
-	};
-
-	return GetLabelsForGame<TFillLabelsPolicy>(GameName);
-}
-
-TSharedPtr<TArray<FString> > FUnrealSync::GetAllLabels()
-{
-	struct TFillLabelsPolicy
-	{
-		static void Fill(TArray<FString>& OutLabelNames, const FString& InGameName, const TArray<FP4Label>& InLabels)
-		{
-			for (auto Label : InLabels)
-			{
-				OutLabelNames.Add(Label.GetName());
-			}
-		}
-	};
-
-	return GetLabelsForGame<TFillLabelsPolicy>("");
-}
-
-TSharedPtr<TArray<FString> > FUnrealSync::GetPossibleGameNames()
-{
-	TSharedPtr<TArray<FString> > PossibleGames = MakeShareable(new TArray<FString>());
-
-	FP4Env& Env = FP4Env::Get();
-
-	FString FileList;
-	if (!Env.RunP4Output(FString::Printf(TEXT("files -e %s/.../Build/ArtistSyncRules.xml"), *Env.GetBranch()), FileList) || FileList.IsEmpty())
-	{
-		return PossibleGames;
-	}
-
-	FString Line, Rest = FileList;
-	while (Rest.Split(LINE_TERMINATOR, &Line, &Rest, ESearchCase::CaseSensitive))
-	{
-		if (!Line.StartsWith(Env.GetBranch()))
-		{
-			continue;
-		}
-
-		int32 ArtistSyncRulesPos = Line.Find("/Build/ArtistSyncRules.xml#", ESearchCase::IgnoreCase);
-
-		if (ArtistSyncRulesPos == INDEX_NONE)
-		{
-			continue;
-		}
-
-		FString MiddlePart = Line.Mid(Env.GetBranch().Len(), ArtistSyncRulesPos - Env.GetBranch().Len());
-
-		int32 LastSlash = INDEX_NONE;
-		MiddlePart.FindLastChar('/', LastSlash);
-
-		PossibleGames->Add(MiddlePart.Mid(LastSlash + 1));
-	}
-
-	return PossibleGames;
-}
-
 const FString& FUnrealSync::GetSharedPromotableDisplayName()
 {
 	static const FString DispName = "Shared Promotable";
@@ -278,168 +98,6 @@ const FString& FUnrealSync::GetSharedPromotableP4FolderName()
 	return DispName;
 }
 
-void FUnrealSync::RegisterOnDataLoaded(const FOnDataLoaded& InOnDataLoaded)
-{
-	FUnrealSync::OnDataLoaded = InOnDataLoaded;
-}
-
-void FUnrealSync::RegisterOnDataReset(const FOnDataReset& InOnDataReset)
-{
-	FUnrealSync::OnDataReset = InOnDataReset;
-}
-
-void FUnrealSync::StartLoadingData()
-{
-	bLoadingFinished = false;
-	Data.Reset();
-	LoaderThread.Reset();
-
-	OnDataReset.ExecuteIfBound();
-
-	LoaderThread = MakeShareable(new FP4DataLoader(
-		FP4DataLoader::FOnLoadingFinished::CreateStatic(&FUnrealSync::OnP4DataLoadingFinished)
-		));
-}
-
-void FUnrealSync::TerminateLoadingProcess()
-{
-	if (LoaderThread.IsValid())
-	{
-		LoaderThread->Terminate();
-	}
-}
-
-void FUnrealSync::OnP4DataLoadingFinished(TSharedPtr<FP4DataCache> InData)
-{
-	FUnrealSync::Data = InData;
-
-	bLoadingFinished = true;
-
-	OnDataLoaded.ExecuteIfBound();
-}
-
-bool FUnrealSync::LoadingFinished()
-{
-	return bLoadingFinished;
-}
-
-bool FUnrealSync::HasValidData()
-{
-	return Data.IsValid();
-}
-
-/**
- * Class to store info of syncing thread.
- */
-class FSyncingThread : public FRunnable
-{
-public:
-	/**
-	 * Constructor
-	 *
-	 * @param InSettings Sync settings.
-	 * @param InLabelNameProvider Label name provider.
-	 * @param InOnSyncFinished Delegate to run when syncing process has finished.
-	 * @param InOnSyncProgress Delegate to run when syncing process has made progress.
-	 */
-	FSyncingThread(FSyncSettings InSettings, ILabelNameProvider& InLabelNameProvider, const FUnrealSync::FOnSyncFinished& InOnSyncFinished, const FUnrealSync::FOnSyncProgress& InOnSyncProgress)
-		: bTerminate(false), Settings(MoveTemp(InSettings)), LabelNameProvider(InLabelNameProvider), OnSyncFinished(InOnSyncFinished), OnSyncProgress(InOnSyncProgress)
-	{
-		Thread = FRunnableThread::Create(this, TEXT("Syncing thread"));
-	}
-
-	/**
-	 * Main thread function.
-	 */
-	uint32 Run() override
-	{
-		if (OnSyncProgress.IsBound())
-		{
-			OnSyncProgress.Execute(LabelNameProvider.GetStartedMessage() + "\n");
-		}
-
-		FString Label = LabelNameProvider.GetLabelName();
-		FString Game = LabelNameProvider.GetGameName();
-
-		struct FProcessStopper
-		{
-			FProcessStopper(bool& bInStop, FUnrealSync::FOnSyncProgress& InOuterSyncProgress)
-				: bStop(bInStop), OuterSyncProgress(InOuterSyncProgress) {}
-
-			bool OnProgress(const FString& Text)
-			{
-				if (OuterSyncProgress.IsBound())
-				{
-					if (!OuterSyncProgress.Execute(Text))
-					{
-						bStop = true;
-					}
-				}
-
-				return !bStop;
-			}
-
-		private:
-			bool& bStop;
-			FUnrealSync::FOnSyncProgress& OuterSyncProgress;
-		};
-		
-		FProcessStopper Stopper(bTerminate, OnSyncProgress);
-		bool bSuccess = FUnrealSync::Sync(Settings, Label, Game, FUnrealSync::FOnSyncProgress::CreateRaw(&Stopper, &FProcessStopper::OnProgress));
-
-		if (OnSyncProgress.IsBound())
-		{
-			OnSyncProgress.Execute(LabelNameProvider.GetFinishedMessage() + "\n");
-		}
-
-		OnSyncFinished.ExecuteIfBound(bSuccess);
-
-		return 0;
-	}
-
-	/**
-	 * Stops process runnning in the background and terminates wait for the
-	 * watcher thread to finish.
-	 */
-	void Terminate()
-	{
-		bTerminate = true;
-		Thread->WaitForCompletion();
-	}
-
-private:
-	/* Tells the thread to terminate the process. */
-	bool bTerminate;
-
-	/* Handle for thread object. */
-	FRunnableThread* Thread;
-
-	/* Sync settings. */
-	FSyncSettings Settings;
-
-	/* Label name provider. */
-	ILabelNameProvider& LabelNameProvider;
-
-	/* Delegate that will be run when syncing process has finished. */
-	FUnrealSync::FOnSyncFinished OnSyncFinished;
-
-	/* Delegate that will be run when syncing process has finished. */
-	FUnrealSync::FOnSyncProgress OnSyncProgress;
-};
-
-void FUnrealSync::TerminateSyncingProcess()
-{
-	if (SyncingThread.IsValid())
-	{
-		SyncingThread->Terminate();
-	}
-}
-
-void FUnrealSync::LaunchSync(FSyncSettings Settings, ILabelNameProvider& LabelNameProvider, const FOnSyncFinished& OnSyncFinished, const FOnSyncProgress& OnSyncProgress)
-{
-	SyncingThread = MakeShareable(new FSyncingThread(MoveTemp(Settings), LabelNameProvider, OnSyncFinished, OnSyncProgress));
-}
-
 /**
  * Helper function to send message through sync progress delegate.
  *
@@ -660,6 +318,35 @@ bool FUnrealSync::Sync(const FSyncSettings& Settings, const FString& Label, cons
 	return true;
 }
 
+void FUnrealSync::SaveGUISettingsToCache()
+{
+	SaveGUISettings();
+}
+
+void FUnrealSync::SaveSettings()
+{
+	FSettingsCache::Get().Save();
+}
+
+void FUnrealSync::SaveSettingsAndClose()
+{
+	SaveGUISettingsToCache();
+
+	SaveSettings();
+
+	FPlatformMisc::RequestExit(false);
+}
+
+void FUnrealSync::SaveSettingsAndRestart()
+{
+	SaveGUISettingsToCache();
+
+	SaveSettings();
+
+	FUnrealSync::RunDetachedUS(FPlatformProcess::ExecutableName(false), true, true, false);
+	FPlatformMisc::RequestExit(false);
+}
+
 bool FUnrealSync::IsDebugParameterSet()
 {
 	const auto bDebug = FParse::Param(FCommandLine::Get(), TEXT("Debug"));
@@ -742,10 +429,4 @@ bool FUnrealSync::Initialization(const TCHAR* CommandLine)
 
 
 /* Static fields initialization. */
-FUnrealSync::FOnDataLoaded FUnrealSync::OnDataLoaded;
-FUnrealSync::FOnDataReset FUnrealSync::OnDataReset;
-TSharedPtr<FP4DataCache> FUnrealSync::Data;
-bool FUnrealSync::bLoadingFinished = false;
-TSharedPtr<FP4DataLoader> FUnrealSync::LoaderThread;
-TSharedPtr<FSyncingThread> FUnrealSync::SyncingThread;
 FString FUnrealSync::InitializationError;
diff --git a/Engine/Source/Programs/UnrealSync/Private/UnrealSync.h b/Engine/Source/Programs/UnrealSync/Private/UnrealSync.h
--- a/Engine/Source/Programs/UnrealSync/Private/UnrealSync.h
+++ b/Engine/Source/Programs/UnrealSync/Private/UnrealSync.h
@@ -5,83 +5,13 @@
 #include "Core.h"
 #include "CoreUObject.h"
 #include "P4DataCache.h"
+#include "LabelNameProvider.h"
 
 /**
  * Class to store date needed for sync monitoring thread.
  */
 class FSyncingThread;
 
-/**
-* Interface of sync command line provider widget.
-*/
-class ILabelNameProvider
-{
-public:
-	/**
-	 * Gets label name from current options picked by user.
-	 *
-	 * @returns Label name.
-	 */
-	virtual FString GetLabelName() const = 0;
-
-	/**
-	 * Gets message to display when sync task has started.
-	 */
-	virtual FString GetStartedMessage() const
-	{
-		return "Sync started.";
-	}
-
-	/**
-	 * Gets message to display when sync task has finished.
-	 */
-	virtual FString GetFinishedMessage() const
-	{
-		return "Sync finished.";
-	}
-
-	/**
-	* Gets game name from current options picked by user.
-	*
-	* @returns Game name.
-	*/
-	virtual FString GetGameName() const
-	{
-		return CurrentGameName;
-	}
-
-	/**
-	 * Reset data.
-	 *
-	 * @param GameName Current game name.
-	 */
-	virtual void ResetData(const FString& GameName)
-	{
-		CurrentGameName = GameName;
-	}
-
-	/**
-	 * Refresh data.
-	 *
-	 * @param GameName Current game name.
-	 */
-	virtual void RefreshData(const FString& GameName)
-	{
-		CurrentGameName = GameName;
-	}
-
-	/**
-	 * Tells if this particular widget has data ready for sync.
-	 *
-	 * @returns True if ready. False otherwise.
-	 */
-	virtual bool IsReadyForSync() const = 0;
-
-private:
-	/* Current game name. */
-	FString CurrentGameName;
-};
-
 struct FSyncSettings
 {
 	/* Artist sync? */
@@ -117,14 +47,6 @@ struct FSyncSettings
 class FUnrealSync
 {
 public:
-	/* On data loaded event delegate. */
-	DECLARE_DELEGATE(FOnDataLoaded);
-
-	/* On data reset event delegate. */
-	DECLARE_DELEGATE(FOnDataReset);
-
-	/* On sync finished event delegate. */
-	DECLARE_DELEGATE_OneParam(FOnSyncFinished, bool);
 	/* On sync log chunk read event delegate. */
 	DECLARE_DELEGATE_RetVal_OneParam(bool, FOnSyncProgress, const FString&);
 
@@ -143,47 +65,6 @@ public:
 	 */
 	static bool Initialization(const TCHAR* CommandLine);
 
-	/**
-	 * Gets latest label for given game name.
-	 *
-	 * @param GameName Current game name.
-	 *
-	 * @returns Latest label for current game.
-	 */
-	static FString GetLatestLabelForGame(const FString& GameName);
-
-	/**
-	 * Gets promoted labels for given game.
-	 *
-	 * @param GameName Current game name.
-	 *
-	 * @returns Array of promoted labels for given game.
-	 */
-	static TSharedPtr<TArray<FString> > GetPromotedLabelsForGame(const FString& GameName);
-
-	/**
-	 * Gets promotable labels for given game since last promoted.
-	 *
-	 * @param GameName Current game name.
-	 *
-	 * @returns Array of promotable labels for given game since last promoted.
-	 */
-	static TSharedPtr<TArray<FString> > GetPromotableLabelsForGame(const FString& GameName);
-
-	/**
-	 * Gets all labels.
-	 *
-	 * @returns Array of all labels names.
-	 */
-	static TSharedPtr<TArray<FString> > GetAllLabels();
-
-	/**
-	 * Gets possible game names.
-	 *
-	 * @returns Array of possible game names.
-	 */
-	static TSharedPtr<TArray<FString> > GetPossibleGameNames();
-
 	/**
 	 * Gets shared promotable display name.
 	 *
@@ -196,25 +77,6 @@ public:
 	 */
 	static const FString& GetSharedPromotableP4FolderName();
 
-	/**
-	 * Registers event that will be trigger when data is loaded.
-	 *
-	 * @param InOnDataLoaded Delegate to call when event happens.
-	 */
-	static void RegisterOnDataLoaded(const FOnDataLoaded& InOnDataLoaded);
-
-	/**
-	 * Registers event that will be trigger when data is reset.
-	 *
-	 * @param InOnDataReset Delegate to call when event happens.
-	 */
-	static void RegisterOnDataReset(const FOnDataReset& InOnDataReset);
-
-	/**
-	 * Start async loading of the P4 label data in case user wants it.
-	 */
-	static void StartLoadingData();
-
 	/**
 	 * Runs detached UnrealSync process and passes given parameters in command line.
 	 *
@@ -236,40 +98,6 @@ public:
 	 */
 	static void RunUnrealSync(const TCHAR* CommandLine);
 
-	/**
-	 * Tells that labels names are currently being loaded.
-	 *
-	 * @returns True if labels names are currently being loaded. False otherwise.
-	 */
-	static bool IsLoadingInProgress()
-	{
-		return LoaderThread.IsValid() && LoaderThread->IsInProgress();
-	}
-
-	/**
-	 * Terminates background P4 data loading process.
-	 */
-	static void TerminateLoadingProcess();
-
-	/**
-	 * Terminates P4 syncing process.
-	 */
-	static void TerminateSyncingProcess();
-
-	/**
-	 * Method to receive p4 data loading finished event.
-	 *
-	 * @param InData Loaded data.
-	 */
-	static void OnP4DataLoadingFinished(TSharedPtr<FP4DataCache> InData);
-
-	/**
-	 * Tells if has valid P4 data loaded.
-	 *
-	 * @returns If P4 data was loaded.
-	 */
-	static bool HasValidData();
-
 	/**
 	 * Tells if loading has finished.
 	 *
@@ -284,23 +112,6 @@ public:
 	 */
 	static const FString& GetInitializationError() { return InitializationError; }
 
-	/**
-	 * Gets labels from the loaded cache.
-	 *
-	 * @returns Array of loaded labels.
-	 */
-	static const TArray<FP4Label>& GetLabels() { return Data->GetLabels(); }
-
-	/**
-	 * Launches UAT UnrealSync command with given command line and options.
-	 *
-	 * @param Settings Sync settings.
-	 * @param LabelNameProvider Object that will provide label name to syncing thread.
-	 * @param OnSyncFinished Delegate to run when syncing is finished.
-	 * @param OnSyncProgress Delegate to run when syncing has made progress.
-	 */
-	static void LaunchSync(FSyncSettings Settings, ILabelNameProvider& LabelNameProvider, const FOnSyncFinished& OnSyncFinished, const FOnSyncProgress& OnSyncProgress);
-
 	/**
 	 * Performs the actual sync with given params.
 	 *
@@ -310,7 +121,28 @@ public:
 	 * @param OnSyncProgress Delegate to run when syncing has made progress. 
 	 */
 	static bool Sync(const FSyncSettings& Settings, const FString& Label, const FString& Game, const FOnSyncProgress& OnSyncProgress);
+
+	/**
+	 * Save GUI settings to cache, flush it to file and close the app.
+	 */
+	static void SaveSettingsAndClose();
+
+	/**
+	 * Save GUI settings to cache, flush it to file and restart the app.
+	 */
+	static void SaveSettingsAndRestart();
+
 private:
+	/**
+	 * Save cached settings to file.
+	 */
+	static void SaveSettings();
+
+	/**
+	 * Save GUI settings to settings' cache.
+	 */
+	static void SaveGUISettingsToCache();
+
 	/**
 	 * Tries to update original UnrealSync at given location.
 	 *
@@ -331,24 +163,6 @@ private:
 	 */
 	static bool DeleteIfExistsAndCopy(const FString& To, const FString& From);
 
-	/* Tells if loading has finished. */
-	static bool bLoadingFinished;
-
-	/* Data loaded event. */
-	static FOnDataLoaded OnDataLoaded;
-
-	/* Data reset event. */
-	static FOnDataReset OnDataReset;
-
-	/* Cached data ptr. */
-	static TSharedPtr<FP4DataCache> Data;
-
-	/* Background loading process monitoring thread. */
-	static TSharedPtr<FP4DataLoader> LoaderThread;
-
-	/* Background syncing process monitoring thread. */
-	static TSharedPtr<FSyncingThread> SyncingThread;
-
 	/* Error that happened during application initialization. */
 	static FString InitializationError;
 };
