Added missing files for last UnrealSync submit.

diff --git a/Engine/Source/Programs/UnrealSync/Private/LabelNameProvider.h b/Engine/Source/Programs/UnrealSync/Private/LabelNameProvider.h
new file mode 100644
--- /dev/null
+++ b/Engine/Source/Programs/UnrealSync/Private/LabelNameProvider.h
@@ -0,0 +1,76 @@
+// Copyright 1998-2015 Epic Games, Inc. All Rights Reserved.
+
+#pragma once
+
+#include "UnrealString.h"
+
+/**
+ * Interface of sync command line provider widget.
+ */
+class ILabelNameProvider
+{
+public:
+	/**
+	 * Gets label name from current options picked by user.
+	 *
+	 * @returns Label name.
+	 */
+	virtual FString GetLabelName() const = 0;
+
+	/**
+	 * Gets message to display when sync task has started.
+	 */
+	virtual FString GetStartedMessage() const
+	{
+		return "Sync started.";
+	}
+
+	/**
+	 * Gets message to display when sync task has finished.
+	 */
+	virtual FString GetFinishedMessage() const
+	{
+		return "Sync finished.";
+	}
+
+	/**
+	* Gets game name from current options picked by user.
+	*
+	* @returns Game name.
+	*/
+	virtual FString GetGameName() const
+	{
+		return CurrentGameName;
+	}
+
+	/**
+	 * Reset data.
+	 *
+	 * @param GameName Current game name.
+	 */
+	virtual void ResetData(const FString& GameName)
+	{
+		CurrentGameName = GameName;
+	}
+
+	/**
+	 * Refresh data.
+	 *
+	 * @param GameName Current game name.
+	 */
+	virtual void RefreshData(const FString& GameName)
+	{
+		CurrentGameName = GameName;
+	}
+
+	/**
+	 * Tells if this particular widget has data ready for sync.
+	 *
+	 * @returns True if ready. False otherwise.
+	 */
+	virtual bool IsReadyForSync() const = 0;
+
+private:
+	/* Current game name. */
+	FString CurrentGameName;
+};
\ No newline at end of file
diff --git a/Engine/Source/Programs/UnrealSync/Private/SyncingThread.cpp b/Engine/Source/Programs/UnrealSync/Private/SyncingThread.cpp
new file mode 100644
--- /dev/null
+++ b/Engine/Source/Programs/UnrealSync/Private/SyncingThread.cpp
@@ -0,0 +1,63 @@
+// Copyright 1998-2015 Epic Games, Inc. All Rights Reserved.
+
+#include "UnrealSync.h"
+
+#include "SyncingThread.h"
+
+FSyncingThread::FSyncingThread(FSyncSettings InSettings, ILabelNameProvider& InLabelNameProvider, const FOnSyncFinished& InOnSyncFinished, const FOnSyncProgress& InOnSyncProgress)
+	: bTerminate(false), Settings(MoveTemp(InSettings)), LabelNameProvider(InLabelNameProvider), OnSyncFinished(InOnSyncFinished), OnSyncProgress(InOnSyncProgress)
+{
+	Thread = FRunnableThread::Create(this, TEXT("Syncing thread"));
+}
+
+uint32 FSyncingThread::Run()
+{
+	if (OnSyncProgress.IsBound())
+	{
+		OnSyncProgress.Execute(LabelNameProvider.GetStartedMessage() + "\n");
+	}
+
+	FString Label = LabelNameProvider.GetLabelName();
+	FString Game = LabelNameProvider.GetGameName();
+
+	struct FProcessStopper
+	{
+		FProcessStopper(bool& bInStop, FOnSyncProgress& InOuterSyncProgress)
+			: bStop(bInStop), OuterSyncProgress(InOuterSyncProgress) {}
+
+		bool OnProgress(const FString& Text)
+		{
+			if (OuterSyncProgress.IsBound())
+			{
+				if (!OuterSyncProgress.Execute(Text))
+				{
+					bStop = true;
+				}
+			}
+
+			return !bStop;
+		}
+
+	private:
+		bool& bStop;
+		FOnSyncProgress& OuterSyncProgress;
+	};
+
+	FProcessStopper Stopper(bTerminate, OnSyncProgress);
+	bool bSuccess = FUnrealSync::Sync(Settings, Label, Game, FUnrealSync::FOnSyncProgress::CreateRaw(&Stopper, &FProcessStopper::OnProgress));
+
+	if (OnSyncProgress.IsBound())
+	{
+		OnSyncProgress.Execute(LabelNameProvider.GetFinishedMessage() + "\n");
+	}
+
+	OnSyncFinished.ExecuteIfBound(bSuccess);
+
+	return 0;
+}
+
+void FSyncingThread::Terminate()
+{
+	bTerminate = true;
+	Thread->WaitForCompletion();
+}
diff --git a/Engine/Source/Programs/UnrealSync/Private/SyncingThread.h b/Engine/Source/Programs/UnrealSync/Private/SyncingThread.h
new file mode 100644
--- /dev/null
+++ b/Engine/Source/Programs/UnrealSync/Private/SyncingThread.h
@@ -0,0 +1,55 @@
+// Copyright 1998-2015 Epic Games, Inc. All Rights Reserved.
+
+#pragma once
+
+/**
+ * Class to store info of syncing thread.
+ */
+class FSyncingThread : public FRunnable
+{
+public:
+	/* On sync finished event delegate. */
+	DECLARE_DELEGATE_OneParam(FOnSyncFinished, bool);
+	/* On sync log chunk read event delegate. */
+	DECLARE_DELEGATE_RetVal_OneParam(bool, FOnSyncProgress, const FString&);
+
+	/**
+	 * Constructor
+	 *
+	 * @param InSettings Sync settings.
+	 * @param InLabelNameProvider Label name provider.
+	 * @param InOnSyncFinished Delegate to run when syncing process has finished.
+	 * @param InOnSyncProgress Delegate to run when syncing process has made progress.
+	 */
+	FSyncingThread(FSyncSettings InSettings, ILabelNameProvider& InLabelNameProvider, const FOnSyncFinished& InOnSyncFinished, const FOnSyncProgress& InOnSyncProgress);
+
+	/**
+	 * Main thread function.
+	 */
+	uint32 Run() override;
+
+	/**
+	 * Stops process running in the background and terminates wait for the
+	 * watcher thread to finish.
+	 */
+	void Terminate();
+
+private:
+	/* Tells the thread to terminate the process. */
+	bool bTerminate;
+
+	/* Handle for thread object. */
+	FRunnableThread* Thread;
+
+	/* Sync settings. */
+	FSyncSettings Settings;
+
+	/* Label name provider. */
+	ILabelNameProvider& LabelNameProvider;
+
+	/* Delegate that will be run when syncing process has finished. */
+	FOnSyncFinished OnSyncFinished;
+
+	/* Delegate that will be run when syncing process has finished. */
+	FOnSyncProgress OnSyncProgress;
+};
\ No newline at end of file
