UnrealSync fixes for clang compilation.

diff --git a/Engine/Source/Programs/UnrealSync/Private/SettingsCache.h b/Engine/Source/Programs/UnrealSync/Private/SettingsCache.h
--- a/Engine/Source/Programs/UnrealSync/Private/SettingsCache.h
+++ b/Engine/Source/Programs/UnrealSync/Private/SettingsCache.h
@@ -94,7 +94,7 @@ public:
 	static TSharedPtr<FJsonValue> GetStateAsJSON(const TPreservableType Object)
 	{
 		typedef TPointerTraits<TPreservableType, true> T;
-		T::ObjectRefType Ref = T::GetObject(Object);
+		typename T::ObjectRefType Ref = T::GetObject(Object);
 
 		return Ref.GetStateAsJSON();
 	}
@@ -109,10 +109,10 @@ public:
 	static void LoadStateFromJSON(TPreservableType Object, const TValueType Value)
 	{
 		typedef TPointerTraits<TPreservableType, false> PT;
-		PT::ObjectRefType Ref = PT::GetObject(Object);
+		typename PT::ObjectRefType Ref = PT::GetObject(Object);
 
 		typedef TPointerTraits<TValueType, true> VT;
-		VT::ObjectRefType ValRef = VT::GetObject(Value);
+		typename VT::ObjectRefType ValRef = VT::GetObject(Value);
 
 		Ref.LoadStateFromJSON(ValRef);
 	}
