UE-16270: UnrealSync3 crashes on launch

The problem was with obtaining singleton shared ptr for SMainTabWidget. Fixed.

diff --git a/Engine/Source/Programs/UnrealSync/Private/GUI.cpp b/Engine/Source/Programs/UnrealSync/Private/GUI.cpp
--- a/Engine/Source/Programs/UnrealSync/Private/GUI.cpp
+++ b/Engine/Source/Programs/UnrealSync/Private/GUI.cpp
@@ -1249,9 +1249,7 @@ public:
 
 	SMainTabWidget()
 		: ExternalThreadsDispatcher(MakeShareable(new FExternalThreadsDispatcher())), P4Data(new FP4DataProxy)
-	{
-		GetSingletonPtr() = StaticCastSharedRef<SMainTabWidget>(this->AsShared());
-	}
+	{ }
 
 	/**
 	 * Gets instance of this class.
@@ -1261,6 +1259,14 @@ public:
 		return GetSingletonPtr();
 	}
 
+	/**
+	 * Creates instance of this class.
+	 */
+	static TSharedRef<SMainTabWidget> CreateInstance()
+	{
+		return SAssignNew(GetSingletonPtr(), SMainTabWidget);
+	}
+
 	SLATE_BEGIN_ARGS(SMainTabWidget) {}
 	SLATE_END_ARGS()
 
@@ -2145,8 +2151,8 @@ TSharedRef<SDockTab> GetMainTab(const FSpawnTabArgs& Args)
 		.ToolTipText(FText::FromString("Sync Unreal Engine tool."));
 
 	MainTab->SetContent(
-		SNew(SMainTabWidget)
-		);
+		SMainTabWidget::CreateInstance()
+	);
 
 	FGlobalTabmanager::Get()->SetActiveTab(MainTab);
 
