UECORE-131: Add option to open Perforce from UnrealSync.

diff --git a/Engine/Source/Programs/UnrealSync/Private/GUI.cpp b/Engine/Source/Programs/UnrealSync/Private/GUI.cpp
--- a/Engine/Source/Programs/UnrealSync/Private/GUI.cpp
+++ b/Engine/Source/Programs/UnrealSync/Private/GUI.cpp
@@ -398,6 +398,16 @@ public:
 		return *OptionsSource[CurrentOption];
 	}
 
+	/**
+	 * Gets current option index.
+	 *
+	 * @returns Index of current option.
+	 */
+	int32 GetCurrentOptionIndex() const
+	{
+		return CurrentOption;
+	}
+
 	/**
 	 * Sets current option for this combo box.
 	 *
@@ -699,11 +709,11 @@ public:
 	 */
 	void LoadStateFromJSON(const FJsonValue& Value)
 	{
-		int32 Chosen = 0;
+		int32 ChosenWidget = 0;
 
-		if (Value.TryGetNumber(Chosen) && Chosen < Items.Num())
+		if (Value.TryGetNumber(ChosenWidget) && ChosenWidget < Items.Num())
 		{
-			ChooseEnabledItem(Chosen);
+			ChooseEnabledItem(ChosenWidget);
 		}
 	}
 
@@ -954,7 +964,7 @@ public:
 	 */
 	TSharedPtr<FJsonValue> GetStateAsJSON() const
 	{
-		if (!LabelsCombo->IsEnabled())
+		if (!LabelsCombo->IsEnabled() || LabelsCombo->GetCurrentOptionIndex() < 0)
 		{
 			return nullptr;
 		}
@@ -1378,6 +1388,16 @@ public:
 								SNew(STextBlock).Text(FText::FromString("Run UE4"))
 							]
 						]
+					+ SHorizontalBox::Slot().AutoWidth()
+						[
+							SNew(SButton).HAlign(HAlign_Right)
+							.OnClicked(this, &SMainTabWidget::RunP4V)
+							.IsEnabled(CheckIfP4VExists())
+							.VAlign(VAlign_Center)
+							[
+								SNew(STextBlock).Text(FText::FromString("Run P4V"))
+							]
+						]
 				]
 			]
 			+ SWidgetSwitcher::Slot()
@@ -1820,6 +1840,30 @@ private:
 		return FReply::Handled();
 	}
 
+	/**
+	 * Fires up P4V for current branch.
+	 */
+	FReply RunP4V()
+	{
+		FP4Env& Env = FP4Env::Get();
+		RunProcess(Env.GetP4VPath(), FString::Printf(TEXT("-p %s -u %s -c %s"), *Env.GetPort(), *Env.GetUser(), *Env.GetClient()));
+		return FReply::Handled();
+	}
+
+	/**
+	 * Checks if P4V exists for current P4 environment.
+	 *
+	 * @returns True if P4V exists, false otherwise.
+	 */
+	bool CheckIfP4VExists()
+	{
+#if PLATFORM_MAC
+		return FPaths::DirectoryExists(FP4Env::Get().GetP4VPath());
+#else
+		return FPaths::FileExists(FP4Env::Get().GetP4VPath());
+#endif
+	}
+
 	/**
 	 * Sync button click function. Gets all widget data, constructs command
 	 * line and launches syncing.
diff --git a/Engine/Source/Programs/UnrealSync/Private/P4Env.cpp b/Engine/Source/Programs/UnrealSync/Private/P4Env.cpp
--- a/Engine/Source/Programs/UnrealSync/Private/P4Env.cpp
+++ b/Engine/Source/Programs/UnrealSync/Private/P4Env.cpp
@@ -255,7 +255,7 @@ private:
 /**
  * P4 path param auto-detection iterator.
  */
-class FP4PathDetectionIterator : public FP4EnvParamDetectionIteratorBase
+class FPathDetectionIterator : public FP4EnvParamDetectionIteratorBase
 {
 public:
 	/**
@@ -263,35 +263,20 @@ public:
 	 *
 	 * @param CommandLine Command line to parse.
 	 */
-	FP4PathDetectionIterator(const TCHAR* CommandLine)
-		: FP4EnvParamDetectionIteratorBase(EP4ParamType::Path, CommandLine), Step(0)
+	FPathDetectionIterator(EP4ParamType Type, const TCHAR* CommandLine, const TArray<FString>& StdLocationPrefixes, const FString& ExecutableName, bool bAppIsADirectory)
+		: FP4EnvParamDetectionIteratorBase(Type, CommandLine), Step(0)
 	{
 #if PLATFORM_WINDOWS
-		const TCHAR* StdLocationPrefixes[] =
-		{
-			TEXT("C:\\Program Files\\Perforce"),
-			TEXT("C:\\Program Files (x86)\\Perforce")
-		};
-
 		static const FString WhereCommand = "where";
-		static const FString P4ExecutableName = "p4.exe";
 #else
-		const TCHAR* StdLocationPrefixes[] =
-		{
-			TEXT("/bin"),
-			TEXT("/usr/local/bin"),
-			TEXT("~/bin")
-		};
-
 		static const FString WhereCommand = "whereis";
-		static const FString P4ExecutableName = "p4";
 #endif
 
 		TArray<FString> LocationCandidates;
 
 		// Tries to detect in standard environment paths.
 		FString WhereOutput;
-		if (RunProcessOutput(WhereCommand, P4ExecutableName, WhereOutput) && !WhereOutput.IsEmpty())
+		if (RunProcessOutput(WhereCommand, ExecutableName, WhereOutput) && !WhereOutput.IsEmpty())
 		{
 			TArray<FString> Outputs;
 			if (WhereOutput.ParseIntoArrayLines(Outputs, true))
@@ -303,23 +288,23 @@ public:
 			}
 		}
 
-		for (const auto* StdLocationPrefix : StdLocationPrefixes)
+		for (const auto& StdLocationPrefix : StdLocationPrefixes)
 		{
-			LocationCandidates.Add(FPaths::Combine(StdLocationPrefix, *P4ExecutableName));
+			LocationCandidates.Add(FPaths::Combine(*StdLocationPrefix, *ExecutableName));
 		}
 
 		for (auto& LocationCandidate : LocationCandidates)
 		{
 			FPlatformMisc::NormalizePath(LocationCandidate);
 			LocationCandidate = FPaths::ConvertRelativePathToFull(LocationCandidate);
-			if (FPaths::FileExists(LocationCandidate))
+			if ((!bAppIsADirectory && FPaths::FileExists(LocationCandidate)) || (bAppIsADirectory && FPaths::DirectoryExists(LocationCandidate)))
 			{
 				AddUniqueLocation(LocationCandidate);
 			}
 		}
 	}
 
-	virtual ~FP4PathDetectionIterator() {}
+	virtual ~FPathDetectionIterator() {}
 
 	/**
 	 * Function that tries to detect P4 executable path.
@@ -375,6 +360,135 @@ private:
 	int32 Step;
 };
 
+/**
+ * P4 path param auto-detection iterator.
+ */
+class FP4PathDetectionIterator : public FPathDetectionIterator
+{
+public:
+	FP4PathDetectionIterator(const TCHAR* CommandLine)
+		: FPathDetectionIterator(EP4ParamType::Path, CommandLine, GetStdLocationPrefixes(), GetExecutableName(), false)
+	{}
+
+private:
+	/**
+	 * Gets array of standard location prefixes for P4.
+	 */
+	TArray<FString> GetStdLocationPrefixes() const
+	{
+#if PLATFORM_WINDOWS
+		const TCHAR* StdLocationPrefixes[] =
+		{
+			TEXT("C:\\Program Files\\Perforce"),
+			TEXT("C:\\Program Files (x86)\\Perforce")
+		};
+#else
+		const TCHAR* StdLocationPrefixes[] =
+		{
+			TEXT("/bin"),
+			TEXT("/usr/local/bin"),
+			TEXT("~/bin")
+		};
+#endif
+
+		TArray<FString> Out;
+		for (const auto* StdLocationPrefix : StdLocationPrefixes)
+		{
+			Out.Add(StdLocationPrefix);
+		}
+
+		return Out;
+	}
+
+	/**
+	 * Gets executable name for P4.
+	 */
+	const FString& GetExecutableName() const
+	{
+#if PLATFORM_WINDOWS
+		static const FString ExecutableName = "p4.exe";
+#else
+		static const FString ExecutableName = "p4";
+#endif
+
+		return ExecutableName;
+	}
+};
+
+/**
+ * P4V path param auto-detection iterator.
+ */
+class FP4VPathDetectionIterator : public FPathDetectionIterator
+{
+public:
+	FP4VPathDetectionIterator(const TCHAR* CommandLine)
+		: FPathDetectionIterator(EP4ParamType::P4VPath, CommandLine, GetStdLocationPrefixes(), GetExecutableName(), GetAppIsADir())
+	{}
+
+private:
+	/**
+	 * Gets array of standard location prefixes for P4V.
+	 */
+	TArray<FString> GetStdLocationPrefixes() const
+	{
+#if PLATFORM_WINDOWS
+		const TCHAR* StdLocationPrefixes[] =
+		{
+			TEXT("C:\\Program Files\\Perforce"),
+			TEXT("C:\\Program Files (x86)\\Perforce")
+		};
+#elif PLATFORM_LINUX
+		const TCHAR* StdLocationPrefixes[] =
+		{
+			TEXT("/bin"),
+			TEXT("/usr/local/bin"),
+			TEXT("~/bin")
+		};
+#elif PLATFORM_MAC
+		const TCHAR* StdLocationPrefixes[] =
+		{
+			TEXT("/Applications")
+		};
+#endif
+
+		TArray<FString> Out;
+		for (const auto* StdLocationPrefix : StdLocationPrefixes)
+		{
+			Out.Add(StdLocationPrefix);
+		}
+
+		return Out;
+	}
+
+	/**
+	 * Gets if P4V app is a directory.
+	 */
+	bool GetAppIsADir()
+	{
+#if PLATFORM_MAC
+		return true;
+#else
+		return false;
+#endif
+	}
+
+	/**
+	 * Gets executable name for P4V.
+	 */
+	const FString& GetExecutableName() const
+	{
+#if PLATFORM_WINDOWS
+		static const FString ExecutableName = "p4v.exe";
+#elif PLATFORM_LINUX
+		static const FString ExecutableName = "p4v";
+#elif PLATFORM_MAC
+		static const FString ExecutableName = "p4v.app";
+#endif
+
+		return ExecutableName;
+	}
+};
+
 #include "XmlParser.h"
 
 /**
@@ -763,6 +877,9 @@ bool FP4Env::Init(const TCHAR* CommandLine)
 		return false;
 	}
 
+	// Auto-detect optional params.
+	TempEnv->AutoDetectMissingOptionalParams(CommandLine);
+
 	FP4Env::Env = TempEnv;
 
 	return true;
@@ -807,6 +924,7 @@ bool FP4Env::RunP4(const FString& CommandLine)
 void FP4Env::SerializeParams(const FSerializationTask& SerializationTask)
 {
 	SerializationTask.ExecuteIfBound(Path, EP4ParamType::Path);
+	SerializationTask.ExecuteIfBound(P4VPath, EP4ParamType::P4VPath);
 	SerializationTask.ExecuteIfBound(Port, EP4ParamType::Port);
 	SerializationTask.ExecuteIfBound(User, EP4ParamType::User);
 	SerializationTask.ExecuteIfBound(Client, EP4ParamType::Client);
@@ -819,6 +937,8 @@ FString FP4Env::GetParamName(EP4ParamType Type)
 	{
 	case EP4ParamType::Path:
 		return "P4PATH";
+	case EP4ParamType::P4VPath:
+		return "P4VPATH";
 	case EP4ParamType::Port:
 		return "P4PORT";
 	case EP4ParamType::User:
@@ -848,9 +968,9 @@ bool FP4Env::AutoDetectMissingParams(const TCHAR* CommandLine)
 		{
 			SetParam(Type, ParamDetectionIteratorsStack.Last()->GetCurrent());
 
-			if (Type != EP4ParamType::Branch)
+			if (Type < (EP4ParamType)((int32)EP4ParamType::OptionalDelimiter - 1))
 			{
-				Type = (EP4ParamType) ((int)Type + 1);
+				Type = (EP4ParamType)((int32)Type + 1);
 				ParamDetectionIteratorsStack.Add(IP4EnvParamDetectionIterator::Create(Type, CommandLine, *this));
 				continue;
 			}
@@ -861,7 +981,7 @@ bool FP4Env::AutoDetectMissingParams(const TCHAR* CommandLine)
 		}
 		else
 		{
-			Type = (EP4ParamType) ((int)Type - 1);
+			Type = (EP4ParamType)((int32)Type - 1);
 			ParamDetectionIteratorsStack.RemoveAt(ParamDetectionIteratorsStack.Num() - 1);
 			--IterationCountdown;
 			if (!IterationCountdown)
@@ -874,6 +994,18 @@ bool FP4Env::AutoDetectMissingParams(const TCHAR* CommandLine)
 	return false;
 }
 
+void FP4Env::AutoDetectMissingOptionalParams(const TCHAR* CommandLine)
+{
+	for (EP4ParamType Type = (EP4ParamType)((int32)EP4ParamType::OptionalDelimiter + 1); Type < EP4ParamType::EndParam; Type = (EP4ParamType)((int32) Type + 1))
+	{
+		TSharedPtr<IP4EnvParamDetectionIterator> Iter = IP4EnvParamDetectionIterator::Create(Type, CommandLine, *this);
+		if (Iter->MoveNext())
+		{
+			SetParam(Type, Iter->GetCurrent());
+		}
+	}
+}
+
 FString FP4Env::GetCommandLine()
 {
 	class CommandLineOutput
@@ -899,6 +1031,11 @@ const FString& FP4Env::GetPath() const
 	return Path;
 }
 
+const FString& FP4Env::GetP4VPath() const
+{
+	return P4VPath;
+}
+
 void FP4Env::SetParam(EP4ParamType Type, const FString& Value)
 {
 	switch (Type)
@@ -906,6 +1043,9 @@ void FP4Env::SetParam(EP4ParamType Type, const FString& Value)
 	case EP4ParamType::Path:
 		Path = Value;
 		break;
+	case EP4ParamType::P4VPath:
+		P4VPath = Value;
+		break;
 	case EP4ParamType::Port:
 		Port = Value;
 		break;
@@ -983,6 +1123,8 @@ const FString& FP4Env::GetParamByType(EP4ParamType Type) const
 	{
 	case EP4ParamType::Path:
 		return GetPath();
+	case EP4ParamType::P4VPath:
+		return GetP4VPath();
 	case EP4ParamType::Port:
 		return GetPort();
 	case EP4ParamType::User:
@@ -1004,6 +1146,8 @@ TSharedPtr<IP4EnvParamDetectionIterator> IP4EnvParamDetectionIterator::Create(EP
 	{
 	case EP4ParamType::Path:
 		return MakeShareable(new FP4PathDetectionIterator(CommandLine));
+	case EP4ParamType::P4VPath:
+		return MakeShareable(new FP4VPathDetectionIterator(CommandLine));
 	case EP4ParamType::Port:
 		return MakeShareable(new FP4PortDetectionIterator(CommandLine));
 	case EP4ParamType::User:
@@ -1197,6 +1341,8 @@ void SP4EnvTabWidget::Construct(const FArguments& InArgs)
 				return LOCTEXT("P4Option_Client", "Workspace name");
 			case EP4ParamType::Path:
 				return LOCTEXT("P4Option_Path", "Path to P4 executable");
+			case EP4ParamType::P4VPath:
+				return LOCTEXT("P4Option_P4VPath", "Path to P4V executable");
 			case EP4ParamType::Port:
 				return LOCTEXT("P4Option_Port", "P4 server address");
 			case EP4ParamType::User:
@@ -1212,8 +1358,9 @@ void SP4EnvTabWidget::Construct(const FArguments& InArgs)
 		TSharedRef<FP4Option> Option;
 	};
 
-	Options.Reserve(4);
+	Options.Reserve(5);
 	Options.Add(MakeShareable(new FP4Option(EP4ParamType::Path)));
+	Options.Add(MakeShareable(new FP4Option(EP4ParamType::P4VPath)));
 	Options.Add(MakeShareable(new FP4Option(EP4ParamType::Port)));
 	Options.Add(MakeShareable(new FP4Option(EP4ParamType::User)));
 	Options.Add(MakeShareable(new FP4Option(EP4ParamType::Client)));
diff --git a/Engine/Source/Programs/UnrealSync/Private/P4Env.h b/Engine/Source/Programs/UnrealSync/Private/P4Env.h
--- a/Engine/Source/Programs/UnrealSync/Private/P4Env.h
+++ b/Engine/Source/Programs/UnrealSync/Private/P4Env.h
@@ -14,7 +14,10 @@ enum class EP4ParamType
 	Port = 1,
 	User = 2,
 	Client = 3,
-	Branch = 4
+	Branch = 4,
+	OptionalDelimiter = 5,
+	P4VPath = 6,
+	EndParam = 7
 };
 
 /**
@@ -172,6 +175,13 @@ public:
 	 */
 	const FString& GetPath() const;
 
+	/**
+	 * Gets P4V path.
+	 *
+	 * @returns P4V path.
+	 */
+	const FString& GetP4VPath() const;
+
 	/**
 	 * Gets P4 user.
 	 *
@@ -269,6 +279,13 @@ private:
 	 */
 	bool AutoDetectMissingParams(const TCHAR* CommandLine);
 
+	/**
+	 * Auto-detect missing optional params.
+	 *
+	 * @param CommandLine Command line that the program was run with.
+	 */
+	void AutoDetectMissingOptionalParams(const TCHAR* CommandLine);
+
 	/**
 	 * Sets param value.
 	 *
@@ -283,6 +300,9 @@ private:
 	/* Path to P4 executable. */
 	FString Path;
 
+	/* Path to P4V executable. */
+	FString P4VPath;
+
 	/* P4 port. */
 	FString Port;
 
