UECORE-132: Add flag to auto launch Editor on successful sync.

diff --git a/Engine/Source/Programs/UnrealSync/Private/GUI.cpp b/Engine/Source/Programs/UnrealSync/Private/GUI.cpp
--- a/Engine/Source/Programs/UnrealSync/Private/GUI.cpp
+++ b/Engine/Source/Programs/UnrealSync/Private/GUI.cpp
@@ -1290,46 +1290,18 @@ public:
 		AddToRadioSelection(LOCTEXT("SyncToChosenPromotable", "Sync to chosen promotable label since last promoted"), SAssignNew(ChosenPromotable, SPickPromotable, P4Data).Title(LOCTEXT("PickPromotableLabel", "Pick promotable label: ")));
 		AddToRadioSelection(LOCTEXT("SyncToAnyChosen", "Sync to any chosen label"), SAssignNew(ChosenLabel, SPickAny, P4Data).Title(LOCTEXT("PickLabel", "Pick label: ")));
 
-		PickGameWidget = SNew(SPickGameWidget, P4Data)
-			.OnGamePicked(this, &SMainTabWidget::OnCurrentGameChanged);
-
-		ArtistSyncCheckBox = SNew(SPreservableCheckBox)
-			[
-				SNew(STextBlock).Text(LOCTEXT("ArtistSync", "Artist sync?"))
-			];
-
-		if (!FUnrealSync::IsDebugParameterSet())
-		{
-			// Checked by default.
-			ArtistSyncCheckBox->ToggleCheckedState();
-		}
-
-		PreviewSyncCheckBox = SNew(SPreservableCheckBox)
-			[
-				SNew(STextBlock).Text(LOCTEXT("PreviewSync", "Preview sync?"))
-			];
-
-		AutoClobberSyncCheckBox = SNew(SPreservableCheckBox)
-			[
-				SNew(STextBlock).Text(LOCTEXT("AutoClobberSync", "Auto-clobber sync?"))
-			];
-
-		GoBackButton = SNew(SButton)
-			.IsEnabled(false)
-			.OnClicked(this, &SMainTabWidget::OnGoBackButtonClick)
-			[
-				SNew(STextBlock).Text(LOCTEXT("GoBack", "Go back"))
-			];
-
 		TSharedPtr<SVerticalBox> MainBox;
 		
-		Switcher = SNew(SWidgetSwitcher)
+		this->ChildSlot
+		[
+			SAssignNew(Switcher, SWidgetSwitcher)
 			+ SWidgetSwitcher::Slot()
 			[
 				SAssignNew(MainBox, SVerticalBox)
 				+ SVerticalBox::Slot().AutoHeight().Padding(5.0f)
 				[
-					PickGameWidget.ToSharedRef()
+					SAssignNew(PickGameWidget, SPickGameWidget, P4Data)
+					.OnGamePicked(this, &SMainTabWidget::OnCurrentGameChanged)
 				]
 				+ SVerticalBox::Slot().VAlign(VAlign_Fill)
 				[
@@ -1354,7 +1326,17 @@ public:
 						SNew(SVerticalBox)
 						+SVerticalBox::Slot()
 						[
-							AutoClobberSyncCheckBox.ToSharedRef()
+							SAssignNew(AutoClobberSyncCheckBox, SPreservableCheckBox)
+							[
+								SNew(STextBlock).Text(LOCTEXT("AutoClobberSync", "Auto-clobber sync?"))
+							]
+						]
+						+ SVerticalBox::Slot()
+						[
+							SAssignNew(RunUE4AfterSyncCheckBox, SPreservableCheckBox)
+							[
+								SNew(STextBlock).Text(LOCTEXT("RunUE4AfterSync", "Run UE4 after sync?"))
+							]
 						]
 					]
 					+ SHorizontalBox::Slot().AutoWidth()
@@ -1362,11 +1344,17 @@ public:
 						SNew(SVerticalBox)
 						+ SVerticalBox::Slot()
 						[
-							ArtistSyncCheckBox.ToSharedRef()
+							SAssignNew(ArtistSyncCheckBox, SPreservableCheckBox)
+							[
+								SNew(STextBlock).Text(LOCTEXT("ArtistSync", "Artist sync?"))
+							]
 						]
 						+ SVerticalBox::Slot()
 						[
-							PreviewSyncCheckBox.ToSharedRef()
+							SAssignNew(PreviewSyncCheckBox, SPreservableCheckBox)
+							[
+								SNew(STextBlock).Text(LOCTEXT("PreviewSync", "Preview sync?"))
+							]
 						]
 					]
 					+ SHorizontalBox::Slot().AutoWidth().Padding(30, 0, 0, 0)
@@ -1433,15 +1421,22 @@ public:
 					]
 					+ SHorizontalBox::Slot().AutoWidth()
 					[
-						GoBackButton.ToSharedRef()
+						SAssignNew(GoBackButton, SButton)
+						.IsEnabled(false)
+						.OnClicked(this, &SMainTabWidget::OnGoBackButtonClick)
+						[
+							SNew(STextBlock).Text(LOCTEXT("GoBack", "Go back"))
+						]
 					]
 				]
-			];
-		
-		this->ChildSlot
-			[
-				Switcher.ToSharedRef()
-			];
+			]
+		];
+
+		if (!FUnrealSync::IsDebugParameterSet())
+		{
+			// Checked by default.
+			ArtistSyncCheckBox->ToggleCheckedState();
+		}
 
 		if (FUnrealSync::IsDebugParameterSet())
 		{
@@ -1496,6 +1491,7 @@ public:
 			Action.Serialize("IsArtist", Object.ArtistSyncCheckBox);
 			Action.Serialize("IsPreview", Object.PreviewSyncCheckBox);
 			Action.Serialize("IsAutoClobber", Object.AutoClobberSyncCheckBox);
+			Action.Serialize("ShouldRunUE4AfterSync", Object.RunUE4AfterSyncCheckBox);
 			Action.Serialize("RadioSelection", Object.RadioSelection);
 			Action.Serialize("LatestPromoted", Object.LatestPromoted);
 			Action.Serialize("ChosenPromoted", Object.ChosenPromoted);
@@ -1940,6 +1936,11 @@ private:
 		Throbber->SetAnimate(SThrobber::EAnimation::None);
 		CancelButton->SetEnabled(false);
 		GoBackButton->SetEnabled(true);
+
+		if (bSuccess && RunUE4AfterSyncCheckBox->IsChecked())
+		{
+			RunUE4();
+		}
 	}
 
 	/**
@@ -2138,6 +2139,8 @@ private:
 	TSharedPtr<SPreservableCheckBox> PreviewSyncCheckBox;
 	/* Check box to tell if this should be a auto-clobber sync. */
 	TSharedPtr<SPreservableCheckBox> AutoClobberSyncCheckBox;
+	/* Check box to tell if UnrealSync should run UE4 after sync. */
+	TSharedPtr<SPreservableCheckBox> RunUE4AfterSyncCheckBox;
 
 	/* External thread requests dispatcher. */
 	TSharedRef<FExternalThreadsDispatcher, ESPMode::ThreadSafe> ExternalThreadsDispatcher;
