US3: Do not crash if P4 settings aren't valid.

diff --git a/Engine/Source/Programs/UnrealSync/Private/P4Env.cpp b/Engine/Source/Programs/UnrealSync/Private/P4Env.cpp
--- a/Engine/Source/Programs/UnrealSync/Private/P4Env.cpp
+++ b/Engine/Source/Programs/UnrealSync/Private/P4Env.cpp
@@ -1448,7 +1448,7 @@ SP4EnvTabWidget::FP4Option::FP4Option(EP4ParamType InType)
 	: Type(InType)
 {
 	FString Param;
-	if (FP4Env::Get().GetSetting(Param, Type))
+	if (FP4Env::IsValid() && FP4Env::Get().GetSetting(Param, Type))
 	{
 		Text = FText::FromString(Param);
 	}
