UnrealSync: Added option to delete stale binaries after successful sync.

diff --git a/Engine/Source/Programs/UnrealSync/Private/GUI.cpp b/Engine/Source/Programs/UnrealSync/Private/GUI.cpp
--- a/Engine/Source/Programs/UnrealSync/Private/GUI.cpp
+++ b/Engine/Source/Programs/UnrealSync/Private/GUI.cpp
@@ -1322,6 +1322,18 @@ public:
 					]
 					+ SHorizontalBox::Slot().HAlign(HAlign_Fill)
 					+ SHorizontalBox::Slot().AutoWidth().VAlign(VAlign_Top).Padding(0.0f, 0.0f, 10.0f, 0.0f)
+					[
+						SNew(SVerticalBox)
+						+ SVerticalBox::Slot()
+						[
+							SAssignNew(DeleteStaleBinariesOnSuccessfulSync, SPreservableCheckBox)
+							[
+								SNew(STextBlock).Text(LOCTEXT("DeleteStaleBinariesOnSuccessfulSync", "Delete stale binaries?"))
+							]
+						]
+						+ SVerticalBox::Slot().AutoHeight()
+					]
+					+ SHorizontalBox::Slot().AutoWidth()
 					[
 						SNew(SVerticalBox)
 						+SVerticalBox::Slot()
@@ -1491,6 +1503,7 @@ public:
 			Action.Serialize("IsArtist", Object.ArtistSyncCheckBox);
 			Action.Serialize("IsPreview", Object.PreviewSyncCheckBox);
 			Action.Serialize("IsAutoClobber", Object.AutoClobberSyncCheckBox);
+			Action.Serialize("ShouldDeleteStaleBinariesOnSuccessfulSync", Object.DeleteStaleBinariesOnSuccessfulSync);
 			Action.Serialize("ShouldRunUE4AfterSync", Object.RunUE4AfterSyncCheckBox);
 			Action.Serialize("RadioSelection", Object.RadioSelection);
 			Action.Serialize("LatestPromoted", Object.LatestPromoted);
@@ -1877,6 +1890,7 @@ private:
 			ArtistSyncCheckBox->IsChecked(),
 			PreviewSyncCheckBox->IsChecked(),
 			AutoClobberSyncCheckBox->IsChecked(),
+			DeleteStaleBinariesOnSuccessfulSync->IsChecked(),
 			OverrideSyncStep);
 
 		Switcher->SetActiveWidgetIndex(1);
@@ -2141,6 +2155,8 @@ private:
 	TSharedPtr<SPreservableCheckBox> AutoClobberSyncCheckBox;
 	/* Check box to tell if UnrealSync should run UE4 after sync. */
 	TSharedPtr<SPreservableCheckBox> RunUE4AfterSyncCheckBox;
+	/* Check box to tell if UnrealSync should delete stale binaries on successful sync. */
+	TSharedPtr<SPreservableCheckBox> DeleteStaleBinariesOnSuccessfulSync;
 
 	/* External thread requests dispatcher. */
 	TSharedRef<FExternalThreadsDispatcher, ESPMode::ThreadSafe> ExternalThreadsDispatcher;
diff --git a/Engine/Source/Programs/UnrealSync/Private/UnrealSync.cpp b/Engine/Source/Programs/UnrealSync/Private/UnrealSync.cpp
--- a/Engine/Source/Programs/UnrealSync/Private/UnrealSync.cpp
+++ b/Engine/Source/Programs/UnrealSync/Private/UnrealSync.cpp
@@ -56,6 +56,157 @@ bool FUnrealSync::DeleteIfExistsAndCopy(const FString& To, const FString& From)
 	return false;
 }
 
+void FUnrealSync::DeleteStaleBinaries(const FOnSyncProgress& OnSyncProgress, bool bPreview)
+{
+	OnSyncProgress.Execute(TEXT("... obtaining P4 workspace have files.") LINE_TERMINATOR);
+
+	FString Haves;
+	FP4Env::RunP4Output(FString::Printf(TEXT("have %s/.../Binaries/..."), *FP4Env::Get().GetBranch()), Haves);
+
+	FRegexPattern HaveFilePattern(TEXT("\\s+//.+#\\d+ - (.+)"));
+	FRegexMatcher Matcher(HaveFilePattern, Haves);
+
+	TArray<FString> HaveFiles;
+
+	while (Matcher.FindNext())
+	{
+		FString HaveFile = Matcher.GetCaptureGroup(1);
+
+		FPaths::NormalizeFilename(HaveFile);
+
+#ifdef PLATFORM_WINDOWS
+		HaveFiles.Add(HaveFile.ToLower());
+#else
+		HaveFiles.Add(HaveFile);
+#endif
+	}
+
+	OnSyncProgress.Execute(TEXT("... searching file system and diffing with P4 have state.") LINE_TERMINATOR);
+
+	IPlatformFile& FS = IPlatformFile::GetPlatformPhysical();
+
+	FString BranchDir = FPaths::ConvertRelativePathToFull(FPaths::Combine(*FPaths::EngineDir(), TEXT("..")));
+	TArray<FString> ToDelete;
+	struct FPotentialStaleBinariesVisitor : public IPlatformFile::FDirectoryVisitor
+	{
+		FPotentialStaleBinariesVisitor(const FString& InBranchDir, const TArray<FString>& InP4HaveFiles, TArray<FString>& OutToDelete)
+			: BranchDir(InBranchDir), P4HaveFiles(InP4HaveFiles), Output(OutToDelete), FS(IPlatformFile::GetPlatformPhysical())
+		{
+			DirsBlacklist.Add(TEXT("Build"));
+			DirsBlacklist.Add(TEXT("Config"));
+			DirsBlacklist.Add(TEXT("Content"));
+			DirsBlacklist.Add(TEXT("DerivedDataCache"));
+			DirsBlacklist.Add(TEXT("Documentation"));
+			DirsBlacklist.Add(TEXT("Intermediate"));
+			DirsBlacklist.Add(TEXT("Saved"));
+			DirsBlacklist.Add(TEXT("Shaders"));
+			DirsBlacklist.Add(TEXT("Source"));
+		}
+
+		virtual bool Visit(const TCHAR* FilenameOrDirectory, bool bIsDirectory)
+		{
+
+#ifdef PLATFORM_WINDOWS
+			const ESearchCase::Type CaseSensitiveness = ESearchCase::IgnoreCase;
+#elif PLATFORM_LINUX || PLATFORM_MAC
+			const ESearchCase::Type CaseSensitiveness = ESearchCase::CaseSensitive;
+#endif
+
+			if (!bIsDirectory)
+			{
+				FString File(FilenameOrDirectory);
+
+				FString Filename = FPaths::GetBaseFilename(File);
+				FString Directory = FPaths::GetPath(File);
+				FString Extension = FPaths::GetExtension(File);
+
+				if (Directory.Contains(TEXT("Binaries"), CaseSensitiveness) &&
+					IsStaleBinaryExtension(Extension) &&
+					(
+						Filename.StartsWith(TEXT("UE4Editor")) ||
+						Filename.StartsWith(TEXT("UE4Game"))
+					))
+				{
+					FPaths::NormalizeFilename(File);
+#ifdef PLATFORM_WINDOWS
+					FString NormCaseFile = File.ToLower();
+#else
+					FString NormCaseFile = File;
+#endif
+					if (!P4HaveFiles.ContainsByPredicate([&NormCaseFile](const FString& Other)
+						{
+							return NormCaseFile.Equals(Other, ESearchCase::CaseSensitive);
+						}))
+					{
+						Output.Add(MoveTemp(File));
+					}
+				}
+			}
+			else
+			{
+				FString DirName = FPaths::GetBaseFilename(FilenameOrDirectory);
+
+				for (const auto& BlacklistedDirName : DirsBlacklist)
+				{
+					if (DirName.Equals(BlacklistedDirName, CaseSensitiveness))
+					{
+						return true;
+					}
+				}
+
+				FS.IterateDirectory(FilenameOrDirectory, *this);
+			}
+
+			return true;
+		}
+
+	private:
+		static bool IsStaleBinaryExtension(const FString& Extension)
+		{
+#ifdef PLATFORM_WINDOWS
+			return Extension.Equals(TEXT("pdb"), ESearchCase::IgnoreCase) ||
+				Extension.Equals(TEXT("exe"), ESearchCase::IgnoreCase) ||
+				Extension.Equals(TEXT("dll"), ESearchCase::IgnoreCase);
+#elif PLATFORM_LINUX
+			return Extension.Equals(TEXT("")) ||
+				Extension.Equals(TEXT("so"));
+#elif PLATFORM_MAC
+			return Extension.Equals(TEXT("")) ||
+				Extension.Equals(TEXT("dsym")) ||
+				Extension.Equals(TEXT("dylib"));
+#endif
+		}
+
+		const FString& BranchDir;
+		TArray<FString>& Output;
+		const TArray<FString>& P4HaveFiles;
+		IPlatformFile& FS;
+
+		TArray<FString> DirsBlacklist;
+	} Visitor(BranchDir, HaveFiles, ToDelete);
+
+	FS.IterateDirectory(*BranchDir, Visitor);
+
+	if (bPreview)
+	{
+		for (const auto& FileToDelete : ToDelete)
+		{
+			OnSyncProgress.Execute(FString::Printf(TEXT("... would delete (preview) stale file %s") LINE_TERMINATOR, *FileToDelete));
+		}
+	}
+	else
+	{
+		for (const auto& FileToDelete : ToDelete)
+		{
+			OnSyncProgress.Execute(FString::Printf(TEXT("... deleting stale file %s") LINE_TERMINATOR, *FileToDelete));
+			if (!FS.DeleteFile(*FileToDelete))
+			{
+				OnSyncProgress.Execute(FString::Printf(TEXT("\t... failed to delete %s") LINE_TERMINATOR, *FileToDelete));
+			}
+		}
+	}
+}
+
 void FUnrealSync::RunUnrealSync(const TCHAR* CommandLine)
 {
 	// start up the main loop
@@ -268,27 +419,26 @@ bool FUnrealSync::Sync(const FSyncSettings& Settings, const FString& Label, cons
 			return Progress.Execute(Text);
 		}
 
-		void ProcessLog()
+		TArray<FString> GetCantClobbers() const
 		{
-			Log = Log.Replace(TEXT("\r"), TEXT(""));
+			TArray<FString> CantClobbers;
+			FString NormalizedLog = Log.Replace(TEXT("\r"), TEXT(""));
 			const FRegexPattern CantClobber(TEXT("Can't clobber writable file ([^\\n]+)"));
 
-			FRegexMatcher Match(CantClobber, Log);
+			FRegexMatcher Match(CantClobber, NormalizedLog);
 
 			while (Match.FindNext())
 			{
 				CantClobbers.Add(Match.GetCaptureGroup(1));
 			}
-		}
 
-		const TArray<FString>& GetCantClobbers() const { return CantClobbers; }
+			return CantClobbers;
+		}
 
 	private:
 		const FOnSyncProgress& Progress;
 
 		FString Log;
-
-		TArray<FString> CantClobbers;
 	};
 
 	FString CommandPrefix = FString("sync ") + (Settings.bPreview ? "-n " : "") + FP4Env::Get().GetBranch();
@@ -302,8 +452,6 @@ bool FUnrealSync::Sync(const FSyncSettings& Settings, const FString& Label, cons
 				return false;
 			}
 
-			ErrorsCollector.ProcessLog();
-
 			FString AutoClobberCommandPrefix("sync -f ");
 			for (const auto& CantClobber : ErrorsCollector.GetCantClobbers())
 			{
@@ -315,6 +463,13 @@ bool FUnrealSync::Sync(const FSyncSettings& Settings, const FString& Label, cons
 		}
 	}
 
+	if (Settings.bDeleteStaleBinaries)
+	{
+		OnSyncProgress.Execute(TEXT("Deleting stale binaries...") LINE_TERMINATOR);
+		DeleteStaleBinaries(OnSyncProgress, Settings.bPreview);
+		OnSyncProgress.Execute(TEXT("Stale binaries deleted.") LINE_TERMINATOR);
+	}
+
 	return true;
 }
 
diff --git a/Engine/Source/Programs/UnrealSync/Private/UnrealSync.h b/Engine/Source/Programs/UnrealSync/Private/UnrealSync.h
--- a/Engine/Source/Programs/UnrealSync/Private/UnrealSync.h
+++ b/Engine/Source/Programs/UnrealSync/Private/UnrealSync.h
@@ -23,13 +23,17 @@ struct FSyncSettings
 	/* Auto-clobber sync? */
 	bool bAutoClobber;
 
+	/* Automatically delete stale binaries? */
+	bool bDeleteStaleBinaries;
+
 	/* Override sync step. If set then it overrides collecting sync steps and uses this one instead. */
 	FString OverrideSyncStep;
 
-	FSyncSettings(bool bInArtist, bool bInPreview, bool bInAutoClobber, FString InOverrideSyncStep = FString())
+	FSyncSettings(bool bInArtist, bool bInPreview, bool bInAutoClobber, bool bInDeleteStaleBinaries, FString InOverrideSyncStep = FString())
 		: bArtist(bInArtist)
 		, bPreview(bInPreview)
 		, bAutoClobber(bInAutoClobber)
+		, bDeleteStaleBinaries(bInDeleteStaleBinaries)
 		, OverrideSyncStep(MoveTemp(InOverrideSyncStep))
 	{ }
 
@@ -37,6 +41,7 @@ struct FSyncSettings
 		: bArtist(Other.bArtist)
 		, bPreview(Other.bPreview)
 		, bAutoClobber(Other.bAutoClobber)
+		, bDeleteStaleBinaries(Other.bDeleteStaleBinaries)
 		, OverrideSyncStep(MoveTemp(Other.OverrideSyncStep))
 	{ }
 };
@@ -163,6 +168,11 @@ private:
 	 */
 	static bool DeleteIfExistsAndCopy(const FString& To, const FString& From);
 
+	/**
+	 * Deletes stale binaries from the workspace.
+	 */
+	static void DeleteStaleBinaries(const FOnSyncProgress& OnSyncProgress, bool bPreview);
+
 	/* Error that happened during application initialization. */
 	static FString InitializationError;
 };
