Fixing CIS warning: mono,  FPotentialStaleBinariesVisitor construction list not matching order of field declarations.

diff --git a/Engine/Source/Programs/UnrealSync/Private/UnrealSync.cpp b/Engine/Source/Programs/UnrealSync/Private/UnrealSync.cpp
--- a/Engine/Source/Programs/UnrealSync/Private/UnrealSync.cpp
+++ b/Engine/Source/Programs/UnrealSync/Private/UnrealSync.cpp
@@ -90,7 +90,7 @@ void FUnrealSync::DeleteStaleBinaries(const FOnSyncProgress& OnSyncProgress, boo
 	struct FPotentialStaleBinariesVisitor : public IPlatformFile::FDirectoryVisitor
 	{
 		FPotentialStaleBinariesVisitor(const FString& InBranchDir, const TArray<FString>& InP4HaveFiles, TArray<FString>& OutToDelete)
-			: BranchDir(InBranchDir), P4HaveFiles(InP4HaveFiles), Output(OutToDelete), FS(IPlatformFile::GetPlatformPhysical())
+			: BranchDir(InBranchDir), Output(OutToDelete), P4HaveFiles(InP4HaveFiles), FS(IPlatformFile::GetPlatformPhysical())
 		{
 			DirsBlacklist.Add(TEXT("Build"));
 			DirsBlacklist.Add(TEXT("Config"));
