UE-15513: UnrealSync fails to find UnrealSync.settings

diff --git a/Engine/Source/Programs/UnrealSync/Private/P4Env.cpp b/Engine/Source/Programs/UnrealSync/Private/P4Env.cpp
--- a/Engine/Source/Programs/UnrealSync/Private/P4Env.cpp
+++ b/Engine/Source/Programs/UnrealSync/Private/P4Env.cpp
@@ -798,7 +798,7 @@ private:
 bool GetCurrentBranch(FString& Output, const FP4Env& Env)
 {
 	FString FilesOutput;
-	if (!RunProcessOutput(Env.GetPath(), FString::Printf(TEXT("-p%s -u%s -c%s files %s"),
+	if (!RunProcessOutput(Env.GetPath(), FString::Printf(TEXT("-p%s -u%s -c%s files \"%s\""),
 		*Env.GetPort(), *Env.GetUser(), *Env.GetClient(), *GetKnownPath()), FilesOutput))
 	{
 		UE_LOG(LogP4Env, Log,
