Changed casing of dSYM.

diff --git a/Engine/Source/Programs/UnrealSync/Private/UnrealSync.cpp b/Engine/Source/Programs/UnrealSync/Private/UnrealSync.cpp
--- a/Engine/Source/Programs/UnrealSync/Private/UnrealSync.cpp
+++ b/Engine/Source/Programs/UnrealSync/Private/UnrealSync.cpp
@@ -172,7 +172,7 @@ void FUnrealSync::DeleteStaleBinaries(const FOnSyncProgress& OnSyncProgress, boo
 				Extension.Equals(TEXT("so"));
 #elif PLATFORM_MAC
 			return Extension.Equals(TEXT("")) ||
-				Extension.Equals(TEXT("dsym")) ||
+				Extension.Equals(TEXT("dSYM")) ||
 				Extension.Equals(TEXT("dylib"));
 #endif
 		}
